Moved WP HLOD Builder to it's own module, WorldPartitionHLODUtilities
Necessary to avoid dragging more geometry tools directly in the Engine module when the new Mesh Approximation merge method will be integrated to HLOD

diff --git a/Engine/Source/Developer/WorldPartitionHLODUtilities/Private/WorldPartitionHLODUtilities.cpp b/Engine/Source/Developer/WorldPartitionHLODUtilities/Private/WorldPartitionHLODUtilities.cpp
new file mode 100644
--- /dev/null
+++ b/Engine/Source/Developer/WorldPartitionHLODUtilities/Private/WorldPartitionHLODUtilities.cpp
@@ -0,0 +1,615 @@
+// Copyright Epic Games, Inc. All Rights Reserved.
+
+#include "WorldPartitionHLODUtilities.h"
+
+#if WITH_EDITOR
+
+#include "WorldPartition/WorldPartition.h"
+#include "WorldPartition/WorldPartitionHandle.h"
+#include "WorldPartition/HLOD/HLODActor.h"
+#include "WorldPartition/HLOD/HLODLayer.h"
+#include "WorldPartition/HLOD/HLODActorDesc.h"
+
+#include "ISMPartition/ISMComponentDescriptor.h"
+
+#include "Algo/ForEach.h"
+#include "Algo/RemoveIf.h"
+#include "Algo/Transform.h"
+#include "Misc/HashBuilder.h"
+#include "Serialization/ArchiveCrc32.h"
+#include "Templates/UniquePtr.h"
+
+#include "IMeshMergeUtilities.h"
+#include "MeshMergeModule.h"
+#include "Modules/ModuleManager.h"
+#include "Engine/StaticMesh.h"
+#include "Engine/HLODProxy.h"
+#include "Materials/Material.h"
+#include "StaticMeshCompiler.h"
+
+#include "LevelInstance/LevelInstanceActor.h"
+#include "LevelInstance/LevelInstanceSubsystem.h"
+
+DEFINE_LOG_CATEGORY_STATIC(LogHLODBuilder, Log, All);
+
+/**
+ * Base class for all HLODBuilders
+ */
+class FHLODBuilder
+{
+public:
+	virtual ~FHLODBuilder() {}
+
+	virtual TArray<UPrimitiveComponent*> CreateComponents(AWorldPartitionHLOD* InHLODActor, const UHLODLayer* InHLODLayer, const TArray<UPrimitiveComponent*>& InSubComponents) = 0;
+
+	void Build(AWorldPartitionHLOD* InHLODActor, const UHLODLayer* InHLODLayer, const TArray<const AActor*>& InSubActors)
+	{
+		TArray<UPrimitiveComponent*> SubComponents = GatherPrimitiveComponents(InSubActors);
+		if (SubComponents.IsEmpty())
+		{
+			return;
+		}
+
+		TArray<UPrimitiveComponent*> HLODPrimitives = CreateComponents(InHLODActor, InHLODLayer, SubComponents);
+		HLODPrimitives.RemoveSwap(nullptr);
+
+		if (!HLODPrimitives.IsEmpty())
+		{
+			InHLODActor->Modify();
+			InHLODActor->SetHLODPrimitives(HLODPrimitives);
+		}
+	}
+
+	static bool LoadSubActors(AWorldPartitionHLOD* InHLODActor, TArray<const AActor*>& OutActors, TArray<FWorldPartitionReference>& OutActorReferences)
+	{
+		OutActors.Reserve(InHLODActor->GetSubActors().Num());
+		OutActorReferences.Reserve(InHLODActor->GetSubActors().Num());
+
+		UWorld* World = InHLODActor->GetWorld();
+		UWorldPartition* WorldPartition = World->GetWorldPartition();
+		check(WorldPartition);
+
+		bool bIsDirty = false;
+
+		// Gather (and potentially load) actors
+		for (const FGuid& SubActorGuid : InHLODActor->GetSubActors())
+		{
+			FWorldPartitionReference ActorRef(WorldPartition, SubActorGuid);
+
+			if (ActorRef.IsValid())
+			{
+				AActor* LoadedActor = ActorRef.Get()->GetActor();
+
+				// Load level instances
+				if (ALevelInstance* LevelInstance = Cast<ALevelInstance>(LoadedActor))
+				{
+					// Wait for level instance loading
+					if (LevelInstance->SupportsLoading())
+					{
+						LevelInstance->GetLevelInstanceSubsystem()->BlockLoadLevelInstance(LevelInstance);
+					}
+				}
+
+				OutActors.Add(LoadedActor);
+				OutActorReferences.Add(MoveTemp(ActorRef));
+			}
+			else
+			{
+				bIsDirty = true;
+			}
+		}
+
+		// Wait for compilation to finish
+		FStaticMeshCompilingManager::Get().FinishAllCompilation();
+
+		return !bIsDirty;
+	}
+
+	static TArray<UPrimitiveComponent*> GatherPrimitiveComponents(const TArray<const AActor*>& InActors)
+	{
+		TArray<UPrimitiveComponent*> PrimitiveComponents;
+
+		auto GatherPrimitivesFromActor = [&PrimitiveComponents](const AActor* Actor, const AActor* ParentActor = nullptr)
+		{
+			const TCHAR* Padding = ParentActor ? TEXT("    ") : TEXT("");
+			UE_LOG(LogHLODBuilder, Verbose, TEXT("%s* Adding components from actor %s"), Padding, *Actor->GetName());
+			for (UActorComponent* SubComponent : Actor->GetComponents())
+			{
+				if (SubComponent && SubComponent->IsHLODRelevant())
+				{
+					if (UPrimitiveComponent* PrimitiveComponent = Cast<UPrimitiveComponent>(SubComponent))
+					{
+						PrimitiveComponents.Add(PrimitiveComponent);
+
+						if (UInstancedStaticMeshComponent* ISMC = Cast<UInstancedStaticMeshComponent>(PrimitiveComponent))
+						{
+							UE_LOG(LogHLODBuilder, Verbose, TEXT("%s    * %s [%d instances]"), Padding, *ISMC->GetStaticMesh()->GetName(), ISMC->GetInstanceCount());
+						}
+						else if (UStaticMeshComponent* SMC = Cast<UStaticMeshComponent>(PrimitiveComponent))
+						{
+							UE_LOG(LogHLODBuilder, Verbose, TEXT("%s    * %s"), Padding, *SMC->GetStaticMesh()->GetName());
+						}
+					}
+					else
+					{
+						UE_LOG(LogHLODBuilder, Warning, TEXT("Component \"%s\" is marked as HLOD-relevant but this type of component currently unsupported."), *SubComponent->GetFullName());
+					}
+				}
+			}
+		};
+
+		TSet<AActor*> UnderlyingActors;
+
+		for (const AActor* Actor : InActors)
+		{
+			// Gather primitives from the Actor
+			GatherPrimitivesFromActor(Actor);
+
+			// Retrieve all underlying actors (ex: all sub actors of a LevelInstance)
+			UnderlyingActors.Reset();
+			Actor->EditorGetUnderlyingActors(UnderlyingActors);
+
+			// Gather primitives from underlying actors
+			for (const AActor* UnderlyingActor : UnderlyingActors)
+			{
+				if (UnderlyingActor->IsHLODRelevant())
+				{
+					GatherPrimitivesFromActor(UnderlyingActor, Actor);
+				}
+			}
+		}
+
+		return PrimitiveComponents;
+	}
+
+	static uint32 ComputeHLODHash(AWorldPartitionHLOD* InHLODActor, const TArray<const AActor*>& InActors)
+	{
+		FArchiveCrc32 Ar;
+
+		// Base key, changing this will force a rebuild of all HLODs
+		FString HLODBaseKey = "5184EABE2BEF440DB5A461554A28A3E4";
+		Ar << HLODBaseKey;
+
+		// HLOD Layer
+		uint32 HLODLayerHash = InHLODActor->GetSubActorsHLODLayer()->GetCRC();
+		UE_LOG(LogHLODBuilder, VeryVerbose, TEXT(" - HLODLayer (%s) = %x"), *InHLODActor->GetSubActorsHLODLayer()->GetName(), HLODLayerHash);
+		Ar << HLODLayerHash;
+
+		// We get the CRC of each component
+		TArray<uint32> ComponentsCRCs;
+		for (UPrimitiveComponent* Component : GatherPrimitiveComponents(InActors))
+		{
+			if (UStaticMeshComponent* StaticMeshComponent = Cast<UStaticMeshComponent>(Component))
+			{
+				uint32 ComponentCRC = 0;
+
+				UE_LOG(LogHLODBuilder, VeryVerbose, TEXT(" - Component \'%s\' from actor \'%s\'"), *Component->GetName(), *Component->GetOwner()->GetName());
+
+				// CRC component
+				uint32 StaticMeshComponentCRC = UHLODProxy::GetCRC(StaticMeshComponent);
+				UE_LOG(LogHLODBuilder, VeryVerbose, TEXT("     - StaticMeshComponent (%s) = %x"), *StaticMeshComponent->GetName(), StaticMeshComponentCRC);
+				ComponentCRC = HashCombine(ComponentCRC, StaticMeshComponentCRC);
+
+				if (UStaticMesh* StaticMesh = StaticMeshComponent->GetStaticMesh())
+				{
+					// CRC static mesh
+					int32 StaticMeshCRC = UHLODProxy::GetCRC(StaticMesh);
+					UE_LOG(LogHLODBuilder, VeryVerbose, TEXT("     - StaticMesh (%s) = %x"), *StaticMesh->GetName(), StaticMeshCRC);
+					ComponentCRC = HashCombine(ComponentCRC, StaticMeshCRC);
+
+					// CRC materials
+					const int32 NumMaterials = StaticMeshComponent->GetNumMaterials();
+					for (int32 MaterialIndex = 0; MaterialIndex < NumMaterials; ++MaterialIndex)
+					{
+						UMaterialInterface* MaterialInterface = StaticMeshComponent->GetMaterial(MaterialIndex);
+						if (MaterialInterface)
+						{
+							uint32 MaterialInterfaceCRC = UHLODProxy::GetCRC(MaterialInterface);
+							UE_LOG(LogHLODBuilder, VeryVerbose, TEXT("     - MaterialInterface (%s) = %x"), *MaterialInterface->GetName(), MaterialInterfaceCRC);
+							ComponentCRC = HashCombine(ComponentCRC, MaterialInterfaceCRC);
+
+							TArray<UTexture*> Textures;
+							MaterialInterface->GetUsedTextures(Textures, EMaterialQualityLevel::High, true, ERHIFeatureLevel::SM5, true);
+							for (UTexture* Texture : Textures)
+							{
+								uint32 TextureCRC = UHLODProxy::GetCRC(Texture);
+								UE_LOG(LogHLODBuilder, VeryVerbose, TEXT("     - Texture (%s) = %x"), *Texture->GetName(), TextureCRC);
+								ComponentCRC = HashCombine(ComponentCRC, TextureCRC);
+							}
+						}
+					}
+				}
+
+				ComponentsCRCs.Add(ComponentCRC);
+			}
+		}
+
+		// Sort the components CRCs to ensure the order of components won't have an impact on the final CRC
+		ComponentsCRCs.Sort();
+
+		// Append all components CRCs
+		Ar << ComponentsCRCs;
+			
+		return Ar.GetCrc();
+	}
+
+	void DisableCollisions(UPrimitiveComponent* Component)
+	{
+		Component->SetCollisionProfileName(UCollisionProfile::NoCollision_ProfileName);
+		Component->SetGenerateOverlapEvents(false);
+		Component->SetCanEverAffectNavigation(false);
+		Component->CanCharacterStepUpOn = ECanBeCharacterBase::ECB_No;
+		Component->SetCanEverAffectNavigation(false);
+		Component->SetCollisionEnabled(ECollisionEnabled::NoCollision);
+	}
+};
+
+
+/**
+ * Build a AWorldPartitionHLOD whose components are ISMC
+ */
+class FHLODBuilder_Instancing : public FHLODBuilder
+{
+	// We want to merge all SMC that are using the same static mesh
+	// However, we must also take material overiddes into account.
+	struct FInstancingKey
+	{
+		FInstancingKey(const UStaticMeshComponent* SMC)
+		{
+			FHashBuilder HashBuilder;
+
+			StaticMesh = SMC->GetStaticMesh();
+			HashBuilder << StaticMesh;
+
+			const int32 NumMaterials = SMC->GetNumMaterials();
+			Materials.Reserve(NumMaterials);
+
+			for (int32 MaterialIndex = 0; MaterialIndex < NumMaterials; ++MaterialIndex)
+			{
+				UMaterialInterface* Material = SMC->GetMaterial(MaterialIndex);
+
+				Materials.Add(Material);
+				HashBuilder << Material;
+			}
+
+			Hash = HashBuilder.GetHash();
+		}
+
+		void ApplyTo(UStaticMeshComponent* SMC) const
+		{
+			// Set static mesh
+			SMC->SetStaticMesh(StaticMesh);
+			
+			// Set material overrides
+			for (int32 MaterialIndex = 0; MaterialIndex < Materials.Num(); ++MaterialIndex)
+			{
+				SMC->SetMaterial(MaterialIndex, Materials[MaterialIndex]);
+			}
+		}
+
+		friend uint32 GetTypeHash(const FInstancingKey& Key)
+		{
+			return Key.Hash;
+		}
+
+		bool operator==(const FInstancingKey& Other) const
+		{
+			return Hash == Other.Hash && StaticMesh == Other.StaticMesh && Materials == Other.Materials;
+		}
+
+	private:
+		UStaticMesh*				StaticMesh;
+		TArray<UMaterialInterface*>	Materials;
+		uint32						Hash;
+	};
+
+	virtual TArray<UPrimitiveComponent*> CreateComponents(AWorldPartitionHLOD* InHLODActor, const UHLODLayer* InHLODLayer, const TArray<UPrimitiveComponent*>& InSubComponents) override
+	{
+		TRACE_CPUPROFILER_EVENT_SCOPE(FHLODBuilder_Instancing::CreateComponents);
+
+		TArray<UPrimitiveComponent*> Components;
+
+		// Gather all meshes to instantiate along with their transforms
+		TMap<FInstancingKey, TArray<UPrimitiveComponent*>> Instances;
+		for (UPrimitiveComponent* Primitive : InSubComponents)
+		{
+			if (UStaticMeshComponent* SMC = Cast<UStaticMeshComponent>(Primitive))
+			{
+				Instances.FindOrAdd(FInstancingKey(SMC)).Add(SMC);
+			}
+		}
+
+		// Create an ISMC for each SM asset we found
+		for (const auto& Entry : Instances)
+		{
+			const FInstancingKey EntryInstancingKey = Entry.Key;
+			const TArray<UPrimitiveComponent*>& EntryComponents = Entry.Value;
+			
+			UInstancedStaticMeshComponent* Component = NewObject<UInstancedStaticMeshComponent>(InHLODActor);
+			EntryInstancingKey.ApplyTo(Component);
+			Component->SetForcedLodModel(Component->GetStaticMesh()->GetNumLODs());
+
+			DisableCollisions(Component);
+
+			// Add all instances
+			for (UPrimitiveComponent* SMC : EntryComponents)
+			{
+				// If we have an ISMC, retrieve all instances
+				if (UInstancedStaticMeshComponent* InstancedStaticMeshComponent = Cast<UInstancedStaticMeshComponent>(SMC))
+				{
+					for (int32 InstanceIdx = 0; InstanceIdx < InstancedStaticMeshComponent->GetInstanceCount(); InstanceIdx++)
+					{
+						FTransform InstanceTransform;
+						InstancedStaticMeshComponent->GetInstanceTransform(InstanceIdx, InstanceTransform, true);
+						Component->AddInstanceWorldSpace(InstanceTransform);
+					}
+				}
+				else
+				{
+					Component->AddInstanceWorldSpace(SMC->GetComponentTransform());
+				}
+			}
+
+			Components.Add(Component);
+		};
+
+		return Components;
+	}
+};
+
+
+/**
+ * Build a merged mesh using geometry from the provided actors
+ */
+class FHLODBuilder_MeshMerge : public FHLODBuilder
+{
+	virtual TArray<UPrimitiveComponent*> CreateComponents(AWorldPartitionHLOD* InHLODActor, const UHLODLayer* InHLODLayer, const TArray<UPrimitiveComponent*>& InSubComponents) override
+	{
+		TRACE_CPUPROFILER_EVENT_SCOPE(FHLODBuilder_MeshMerge::CreateComponents);
+
+		TArray<UObject*> Assets;
+		FVector MergedActorLocation;
+
+		const IMeshMergeUtilities& MeshMergeUtilities = FModuleManager::Get().LoadModuleChecked<IMeshMergeModule>("MeshMergeUtilities").GetUtilities();
+		MeshMergeUtilities.MergeComponentsToStaticMesh(InSubComponents, InHLODActor->GetWorld(), InHLODLayer->GetMeshMergeSettings(), InHLODLayer->GetHLODMaterial().LoadSynchronous(), InHLODActor->GetPackage(), InHLODActor->GetActorLabel(), Assets, MergedActorLocation, 0.25f, false);
+
+		UStaticMeshComponent* Component = nullptr;
+		Algo::ForEach(Assets, [this, InHLODActor, &Component, &MergedActorLocation](UObject* Asset)
+		{
+			Asset->ClearFlags(RF_Public | RF_Standalone);
+
+			if (Cast<UStaticMesh>(Asset))
+			{
+				Component = NewObject<UStaticMeshComponent>(InHLODActor);
+				Component->SetStaticMesh(static_cast<UStaticMesh*>(Asset));
+				Component->SetWorldLocation(MergedActorLocation);
+				DisableCollisions(Component);
+			}
+		});
+
+		return TArray<UPrimitiveComponent*>({ Component });
+	}
+};
+
+/**
+ * Build a simplified mesh using geometry from the provided actors
+ */
+class FHLODBuilder_MeshSimplify : public FHLODBuilder
+{
+	virtual TArray<UPrimitiveComponent*> CreateComponents(AWorldPartitionHLOD* InHLODActor, const UHLODLayer* InHLODLayer, const TArray<UPrimitiveComponent*>& InSubComponents) override
+	{
+		TRACE_CPUPROFILER_EVENT_SCOPE(FHLODBuilder_MeshSimplify::CreateComponents);
+
+		TArray<UObject*> Assets;
+		FCreateProxyDelegate ProxyDelegate;
+		ProxyDelegate.BindLambda([&Assets](const FGuid Guid, TArray<UObject*>& InAssetsCreated) { Assets = InAssetsCreated; });
+
+		TArray<UStaticMeshComponent*> StaticMeshComponents;
+		Algo::TransformIf(InSubComponents, StaticMeshComponents, [](UPrimitiveComponent* InPrimitiveComponent) { return InPrimitiveComponent->IsA<UStaticMeshComponent>(); }, [](UPrimitiveComponent* InPrimitiveComponent) { return Cast<UStaticMeshComponent>(InPrimitiveComponent); });
+
+		const IMeshMergeUtilities& MeshMergeUtilities = FModuleManager::Get().LoadModuleChecked<IMeshMergeModule>("MeshMergeUtilities").GetUtilities();
+		MeshMergeUtilities.CreateProxyMesh(StaticMeshComponents, InHLODLayer->GetMeshSimplifySettings(), InHLODLayer->GetHLODMaterial().LoadSynchronous(), InHLODActor->GetPackage(), InHLODActor->GetActorLabel(), FGuid::NewGuid(), ProxyDelegate, true);
+
+		UStaticMeshComponent* Component = nullptr;
+		Algo::ForEach(Assets, [this, InHLODActor, &Component](UObject* Asset)
+		{
+			Asset->ClearFlags(RF_Public | RF_Standalone);
+
+			if (Cast<UStaticMesh>(Asset))
+			{
+				Component = NewObject<UStaticMeshComponent>(InHLODActor);
+				Component->SetStaticMesh(static_cast<UStaticMesh*>(Asset));
+				DisableCollisions(Component);
+			}
+		});
+
+		return TArray<UPrimitiveComponent*>({ Component });
+	}
+};
+
+TArray<AWorldPartitionHLOD*> FWorldPartitionHLODUtilities::CreateHLODActors(FHLODCreationContext& InCreationContext, const FHLODCreationParams& InCreationParams, const TSet<FActorInstance>& InActors, const TArray<const UDataLayer*>& InDataLayers)
+{
+	TMap<UHLODLayer*, TArray<FGuid>> HLODLayersActors;
+	for (const FActorInstance& ActorInstance : InActors)
+	{
+		FWorldPartitionActorDesc& ActorDesc = InCreationParams.WorldPartition->GetActorDescChecked(ActorInstance.Actor);
+		if (ActorDesc.GetActorIsHLODRelevant())
+		{
+			UHLODLayer* HLODLayer = UHLODLayer::GetHLODLayer(ActorDesc, InCreationParams.WorldPartition);
+			if (HLODLayer)
+			{
+				HLODLayersActors.FindOrAdd(HLODLayer).Add(ActorInstance.Actor);
+			}
+		}
+	}
+
+	TArray<AWorldPartitionHLOD*> HLODActors;
+	for (const auto& Pair : HLODLayersActors)
+	{
+		const UHLODLayer* HLODLayer = Pair.Key;
+		const TArray<FGuid> HLODLayerActors = Pair.Value;
+		check(!HLODLayerActors.IsEmpty());
+
+		// Compute HLODActor hash
+		uint64 CellHash = FHLODActorDesc::ComputeCellHash(HLODLayer->GetName(), InCreationParams.GridIndexX, InCreationParams.GridIndexY, InCreationParams.GridIndexZ, InCreationParams.DataLayersID);
+
+		AWorldPartitionHLOD* HLODActor = nullptr;
+		FWorldPartitionHandle HLODActorHandle;
+		if (InCreationContext.HLODActorDescs.RemoveAndCopyValue(CellHash, HLODActorHandle))
+		{
+			InCreationContext.ActorReferences.Add(HLODActorHandle);
+			HLODActor = CastChecked<AWorldPartitionHLOD>(HLODActorHandle->GetActor());
+		}
+
+		bool bNewActor = HLODActor == nullptr;
+		if (bNewActor)
+		{
+			FActorSpawnParameters SpawnParams;
+			SpawnParams.Name = *FString::Printf(TEXT("%s_%016llx"), *HLODLayer->GetName(), CellHash);
+			SpawnParams.NameMode = FActorSpawnParameters::ESpawnActorNameMode::Required_Fatal;
+			HLODActor = InCreationParams.WorldPartition->GetWorld()->SpawnActor<AWorldPartitionHLOD>(SpawnParams);
+
+			FSoftObjectPath WorldPartitionPath(InCreationParams.WorldPartition);
+			FString CellObjectPath = FString::Printf(TEXT("%s.%s"), *WorldPartitionPath.ToString(), *InCreationParams.CellName.ToString());
+			TSoftObjectPtr<UWorldPartitionRuntimeCell> RuntimeCell(CellObjectPath);
+
+			HLODActor->SetActorLabel(FString::Printf(TEXT("HLOD%d_%s"), InCreationParams.HLODLevel, *InCreationParams.CellName.ToString()));
+			HLODActor->SetFolderPath(*FString::Printf(TEXT("HLOD/HLOD%d"), InCreationParams.HLODLevel));
+			HLODActor->SetSourceCell(RuntimeCell);
+			HLODActor->SetSubActorsHLODLayer(HLODLayer);
+			HLODActor->SetGridIndices(InCreationParams.GridIndexX, InCreationParams.GridIndexY, InCreationParams.GridIndexZ);
+
+			// Make sure the generated HLOD actor has the same data layers as the source actors
+			for (const UDataLayer* DataLayer : InDataLayers)
+			{
+				HLODActor->AddDataLayer(DataLayer);
+			}
+		}
+		else
+		{
+#if DO_CHECK
+			uint64 GridIndexX, GridIndexY, GridIndexZ;
+			HLODActor->GetGridIndices(GridIndexX, GridIndexY, GridIndexZ);
+			check(GridIndexX == InCreationParams.GridIndexX);
+			check(GridIndexY == InCreationParams.GridIndexY);
+			check(GridIndexZ == InCreationParams.GridIndexZ);
+			check(HLODActor->GetSubActorsHLODLayer() == HLODLayer);
+			check(FDataLayersID(HLODActor->GetDataLayerObjects()) == InCreationParams.DataLayersID);
+#endif
+		}
+
+		bool bIsDirty = false;
+
+		// Sub actors
+		{
+			bool bSubActorsChanged = HLODActor->GetSubActors().Num() != HLODLayerActors.Num();
+			if (!bSubActorsChanged)
+			{
+				TArray<FGuid> A = HLODActor->GetSubActors();
+				TArray<FGuid> B = HLODLayerActors;
+				A.Sort();
+				B.Sort();
+				bSubActorsChanged = A != B;
+			}
+
+			if (bSubActorsChanged)
+			{
+				HLODActor->SetSubActors(HLODLayerActors);
+				bIsDirty = true;
+			}
+		}
+
+		// Runtime grid
+		FName RuntimeGrid = HLODLayer->GetRuntimeGrid(InCreationParams.HLODLevel);
+		if (HLODActor->GetRuntimeGrid() != RuntimeGrid)
+		{
+			HLODActor->SetRuntimeGrid(RuntimeGrid);
+			bIsDirty = true;
+		}
+
+		// HLOD level
+		if (HLODActor->GetLODLevel() != InCreationParams.HLODLevel)
+		{
+			HLODActor->SetLODLevel(InCreationParams.HLODLevel);
+			bIsDirty = true;
+		}
+
+		// Parent HLOD layer
+		UHLODLayer* ParentHLODLayer = HLODLayer->GetParentLayer().LoadSynchronous();
+		if (HLODActor->GetHLODLayer() != ParentHLODLayer)
+		{
+			HLODActor->SetHLODLayer(ParentHLODLayer);
+			bIsDirty = true;
+		}
+
+		// Cell bounds
+		if (!HLODActor->GetHLODBounds().Equals(InCreationParams.CellBounds))
+		{
+			HLODActor->SetHLODBounds(InCreationParams.CellBounds);
+			bIsDirty = true;
+		}
+
+		// If any change was performed, mark HLOD package as dirty
+		if (bIsDirty)
+		{
+			HLODActor->MarkPackageDirty();
+		}
+
+		HLODActors.Add(HLODActor);
+	}
+
+	return HLODActors;
+}
+
+uint32 FWorldPartitionHLODUtilities::BuildHLOD(AWorldPartitionHLOD* InHLODActor)
+{
+	TArray<FWorldPartitionReference> ActorReferences;
+	TArray<const AActor*> SubActors;
+	
+	bool bIsDirty = !FHLODBuilder::LoadSubActors(InHLODActor, SubActors, ActorReferences);
+	if (bIsDirty)
+	{
+		UE_LOG(LogHLODBuilder, Warning, TEXT("HLOD actor \"%s\" needs to be rebuilt as it references actors that have been deleted."), *InHLODActor->GetActorLabel());
+	}
+
+	uint32 OldHLODHash = bIsDirty ? 0 : InHLODActor->GetHLODHash();
+	uint32 NewHLODHash = FHLODBuilder::ComputeHLODHash(InHLODActor, SubActors);
+
+	if (OldHLODHash == NewHLODHash)
+	{
+		return OldHLODHash;
+	}
+
+	TUniquePtr<FHLODBuilder> HLODBuilder = nullptr;
+	
+	const UHLODLayer* HLODLayer = InHLODActor->GetSubActorsHLODLayer();
+	EHLODLayerType HLODLevelType = HLODLayer->GetLayerType();
+	switch (HLODLevelType)
+	{
+	case EHLODLayerType::Instancing:
+		HLODBuilder = TUniquePtr<FHLODBuilder>(new FHLODBuilder_Instancing());
+		break;
+
+	case EHLODLayerType::MeshMerge:
+		HLODBuilder = TUniquePtr<FHLODBuilder>(new FHLODBuilder_MeshMerge());
+		break;
+
+	case EHLODLayerType::MeshSimplify:
+		HLODBuilder = TUniquePtr<FHLODBuilder>(new FHLODBuilder_MeshSimplify());
+		break;
+
+	default:
+		checkf(false, TEXT("Unsupported type"));
+	}
+
+	if (ensure(HLODBuilder))
+	{
+		HLODBuilder->Build(InHLODActor, HLODLayer, SubActors);
+	}
+			
+	InHLODActor->MarkPackageDirty();
+
+	return NewHLODHash;
+}
+
+#endif
diff --git a/Engine/Source/Developer/WorldPartitionHLODUtilities/Private/WorldPartitionHLODUtilitiesModule.cpp b/Engine/Source/Developer/WorldPartitionHLODUtilities/Private/WorldPartitionHLODUtilitiesModule.cpp
new file mode 100644
--- /dev/null
+++ b/Engine/Source/Developer/WorldPartitionHLODUtilities/Private/WorldPartitionHLODUtilitiesModule.cpp
@@ -0,0 +1,31 @@
+// Copyright Epic Games, Inc. All Rights Reserved.
+
+#include "WorldPartitionHLODUtilitiesModule.h"
+#include "Modules/ModuleManager.h"
+#include "WorldPartitionHLODUtilities.h"
+
+IMPLEMENT_MODULE(FWorldPartitionHLODUtilitiesModule, WorldPartitionHLODUtilities);
+
+void FWorldPartitionHLODUtilitiesModule::StartupModule()
+{
+	Utilities = nullptr;
+}
+
+void FWorldPartitionHLODUtilitiesModule::ShutdownModule()
+{
+	if (Utilities)
+	{
+		delete Utilities;
+		Utilities = nullptr;
+	}
+}
+
+IWorldPartitionHLODUtilities* FWorldPartitionHLODUtilitiesModule::GetUtilities()
+{
+	if (Utilities == nullptr)
+	{
+		Utilities = new FWorldPartitionHLODUtilities();
+	}
+
+	return Utilities;
+}
diff --git a/Engine/Source/Developer/WorldPartitionHLODUtilities/Public/IWorldPartitionHLODUtilities.h b/Engine/Source/Developer/WorldPartitionHLODUtilities/Public/IWorldPartitionHLODUtilities.h
new file mode 100644
--- /dev/null
+++ b/Engine/Source/Developer/WorldPartitionHLODUtilities/Public/IWorldPartitionHLODUtilities.h
@@ -0,0 +1,63 @@
+// Copyright Epic Games, Inc. All Rights Reserved.
+
+#pragma once
+
+#include "CoreMinimal.h"
+#include "Containers/Array.h"
+#include "WorldPartition/WorldPartitionHandle.h"
+#include "WorldPartition/DataLayer/DataLayersID.h"
+#include "WorldPartition/WorldPartitionActorCluster.h"
+
+class AActor;
+class UWorldPartition;
+class UHLODLayer;
+class AWorldPartitionHLOD;
+
+struct ENGINE_API FHLODCreationContext
+{
+	TMap<uint64, FWorldPartitionHandle> HLODActorDescs;
+	TArray<FWorldPartitionReference> ActorReferences;
+};
+
+struct ENGINE_API FHLODCreationParams
+{
+	UWorldPartition* WorldPartition;
+
+	// Everything needed to build the cell hash
+	int64 GridIndexX;
+	int64 GridIndexY;
+	int64 GridIndexZ;
+	FDataLayersID DataLayersID;
+	FName HLODLayerName;
+
+	FName CellName;
+	FBox  CellBounds;
+	uint32 HLODLevel;
+};
+
+/**
+ * Tools for building HLODs in WorldPartition
+ */
+class WORLDPARTITIONHLODUTILITIES_API IWorldPartitionHLODUtilities
+{
+public:
+	virtual ~IWorldPartitionHLODUtilities() {}
+
+	/**
+	 * Create HLOD actors for a given cell
+	 *
+	 * @param	InCreationContext	HLOD creation context object
+	 * @param	InCreationParams	HLOD creation parameters object
+	 * @param	InActors			The actors for which we'll build an HLOD representation
+	 * @param	InDataLayers		The data layers to assign to the newly created HLOD actors
+	 */
+	virtual TArray<AWorldPartitionHLOD*> CreateHLODActors(FHLODCreationContext& InCreationContext, const FHLODCreationParams& InCreationParams, const TSet<FActorInstance>& InActors, const TArray<const UDataLayer*>& InDataLayers) = 0;
+
+	/**
+	 * Build HLOD for the specified AWorldPartitionHLOD actor.
+	 *
+	 * @param 	InHLODActor		The HLOD actor for which we'll build the HLOD
+	 * @return An hash that represent the content used to build this HLOD.
+	 */
+	virtual uint32 BuildHLOD(AWorldPartitionHLOD* InHLODActor) = 0;
+};
diff --git a/Engine/Source/Developer/WorldPartitionHLODUtilities/Public/WorldPartitionHLODUtilities.h b/Engine/Source/Developer/WorldPartitionHLODUtilities/Public/WorldPartitionHLODUtilities.h
new file mode 100644
--- /dev/null
+++ b/Engine/Source/Developer/WorldPartitionHLODUtilities/Public/WorldPartitionHLODUtilities.h
@@ -0,0 +1,15 @@
+// Copyright Epic Games, Inc. All Rights Reserved.
+
+#pragma once
+
+#include "IWorldPartitionHLODUtilities.h"
+
+/**
+* FWorldPartitionHLODUtilities implementation
+*/
+class FWorldPartitionHLODUtilities : public IWorldPartitionHLODUtilities
+{
+public:
+	virtual TArray<AWorldPartitionHLOD*> CreateHLODActors(FHLODCreationContext& InCreationContext, const FHLODCreationParams& InCreationParams, const TSet<FActorInstance>& InActors, const TArray<const UDataLayer*>& InDataLayers) override;
+	virtual uint32 BuildHLOD(AWorldPartitionHLOD* InHLODActor) override;
+};
diff --git a/Engine/Source/Developer/WorldPartitionHLODUtilities/Public/WorldPartitionHLODUtilitiesModule.h b/Engine/Source/Developer/WorldPartitionHLODUtilities/Public/WorldPartitionHLODUtilitiesModule.h
new file mode 100644
--- /dev/null
+++ b/Engine/Source/Developer/WorldPartitionHLODUtilities/Public/WorldPartitionHLODUtilitiesModule.h
@@ -0,0 +1,29 @@
+// Copyright Epic Games, Inc. All Rights Reserved.
+
+#pragma once
+
+#include "CoreMinimal.h"
+#include "Modules/ModuleInterface.h"
+
+class IWorldPartitionHLODUtilities;
+
+class IWorldPartitionHLODUtilitiesModule : public IModuleInterface
+{
+public:
+	virtual IWorldPartitionHLODUtilities* GetUtilities() = 0;
+};
+
+/**
+* IWorldPartitionHLODUtilities module interface
+*/
+class WORLDPARTITIONHLODUTILITIES_API FWorldPartitionHLODUtilitiesModule : public IWorldPartitionHLODUtilitiesModule
+{
+public:
+	virtual void ShutdownModule() override;
+	virtual void StartupModule() override;
+
+	virtual IWorldPartitionHLODUtilities* GetUtilities() override;
+
+private:
+	IWorldPartitionHLODUtilities* Utilities;
+};
diff --git a/Engine/Source/Developer/WorldPartitionHLODUtilities/WorldPartitionHLODUtilities.Build.cs b/Engine/Source/Developer/WorldPartitionHLODUtilities/WorldPartitionHLODUtilities.Build.cs
new file mode 100644
--- /dev/null
+++ b/Engine/Source/Developer/WorldPartitionHLODUtilities/WorldPartitionHLODUtilities.Build.cs
@@ -0,0 +1,45 @@
+// Copyright Epic Games, Inc. All Rights Reserved.
+using UnrealBuildTool;
+
+public class WorldPartitionHLODUtilities : ModuleRules
+{
+    public WorldPartitionHLODUtilities(ReadOnlyTargetRules Target) : base(Target)
+	{
+        PublicIncludePaths.Add("Developer/WorldPartitionHLODUtilities/Public");
+
+        PublicDependencyModuleNames.AddRange(
+            new string[]
+			{
+				"Core",
+				"CoreUObject"
+			}
+        );
+
+        PrivateDependencyModuleNames.AddRange(
+            new string[]
+			{
+				"EditorFramework",
+				"Engine",
+				"MeshDescription",
+				"StaticMeshDescription",
+				"UnrealEd",
+                "Projects",
+			}
+        );
+
+        PrivateIncludePaths.AddRange(
+            new string[]
+            {
+            }
+        );
+
+        DynamicallyLoadedModuleNames.AddRange(
+            new string[]
+            {
+                "MeshUtilities",
+                "MeshMergeUtilities",
+                "MeshReductionInterface",
+            }
+        );
+	}
+}
diff --git a/Engine/Source/Runtime/Engine/Engine.Build.cs b/Engine/Source/Runtime/Engine/Engine.Build.cs
--- a/Engine/Source/Runtime/Engine/Engine.Build.cs
+++ b/Engine/Source/Runtime/Engine/Engine.Build.cs
@@ -234,6 +234,9 @@ public class Engine : ModuleRules
 			
 			PublicIncludePathModuleNames.Add("AnimationDataController");
 			DynamicallyLoadedModuleNames.Add("AnimationDataController");
+			
+			PrivateIncludePathModuleNames.Add("WorldPartitionHLODUtilities");
+			DynamicallyLoadedModuleNames.Add("WorldPartitionHLODUtilities");
 
 			PrivateDependencyModuleNames.AddRange(
 				new string[] {
diff --git a/Engine/Source/Runtime/Engine/Private/WorldPartition/HLOD/HLODBuilder.cpp b/Engine/Source/Runtime/Engine/Private/WorldPartition/HLOD/HLODBuilder.cpp
deleted file mode 100644
diff --git a/Engine/Source/Runtime/Engine/Private/WorldPartition/HLOD/HLODBuilder.h b/Engine/Source/Runtime/Engine/Private/WorldPartition/HLOD/HLODBuilder.h
deleted file mode 100644
diff --git a/Engine/Source/Runtime/Engine/Private/WorldPartition/RuntimeSpatialHash/RuntimeSpatialHashHLOD.cpp b/Engine/Source/Runtime/Engine/Private/WorldPartition/RuntimeSpatialHash/RuntimeSpatialHashHLOD.cpp
--- a/Engine/Source/Runtime/Engine/Private/WorldPartition/RuntimeSpatialHash/RuntimeSpatialHashHLOD.cpp
+++ b/Engine/Source/Runtime/Engine/Private/WorldPartition/RuntimeSpatialHash/RuntimeSpatialHashHLOD.cpp
@@ -5,7 +5,7 @@
 #if WITH_EDITOR
 
 #include "WorldPartition/RuntimeSpatialHash/RuntimeSpatialHashGridHelper.h"
-#include "WorldPartition/HLOD/HLODBuilder.h"
+
 #include "WorldPartition/HLOD/HLODLayer.h"
 #include "WorldPartition/HLOD/HLODActor.h"
 #include "WorldPartition/HLOD/HLODActorDesc.h"
@@ -22,6 +22,9 @@
 
 #include "StaticMeshCompiler.h"
 
+#include "IWorldPartitionHLODUtilities.h"
+#include "WorldPartitionHLODUtilitiesModule.h"
+
 
 DEFINE_LOG_CATEGORY_STATIC(LogWorldPartitionRuntimeSpatialHashHLOD, Log, All);
 
@@ -171,7 +174,9 @@ static TArray<FGuid> GenerateHLODsForGrid(UWorldPartition* WorldPartition, const
 				CreationParams.CellBounds = CellBounds;
 				CreationParams.HLODLevel = HLODLevel;
 
-				TArray<AWorldPartitionHLOD*> CellHLODActors = FHLODBuilderUtilities::CreateHLODActors(Context, CreationParams, GridCellDataChunk.GetActors(), GridCellDataChunk.GetDataLayers());
+				IWorldPartitionHLODUtilities* WPHLODUtilities = FModuleManager::Get().LoadModuleChecked<IWorldPartitionHLODUtilitiesModule>("WorldPartitionHLODUtilities").GetUtilities();
+
+				TArray<AWorldPartitionHLOD*> CellHLODActors = WPHLODUtilities->CreateHLODActors(Context, CreationParams, GridCellDataChunk.GetActors(), GridCellDataChunk.GetDataLayers());
 				if (!CellHLODActors.IsEmpty())
 				{
 					TArray<AWorldPartitionHLOD*> NewCellHLODActors;
diff --git a/Engine/Source/Runtime/Engine/Public/WorldPartition/DataLayer/DataLayersID.h b/Engine/Source/Runtime/Engine/Public/WorldPartition/DataLayer/DataLayersID.h
--- a/Engine/Source/Runtime/Engine/Public/WorldPartition/DataLayer/DataLayersID.h
+++ b/Engine/Source/Runtime/Engine/Public/WorldPartition/DataLayer/DataLayersID.h
@@ -5,7 +5,7 @@
 #include "WorldPartition/DataLayer/DataLayer.h"
 
 #if WITH_EDITOR
-class FDataLayersID
+class ENGINE_API FDataLayersID
 {
 public:
 	FDataLayersID();
