WorldPartition HLOD - First pass at hooking MeshApproximation merge method to HLODs

diff --git a/Engine/Source/Developer/WorldPartitionHLODUtilities/Private/HLODBuilder.cpp b/Engine/Source/Developer/WorldPartitionHLODUtilities/Private/HLODBuilder.cpp
new file mode 100644
--- /dev/null
+++ b/Engine/Source/Developer/WorldPartitionHLODUtilities/Private/HLODBuilder.cpp
@@ -0,0 +1,98 @@
+// Copyright Epic Games, Inc. All Rights Reserved.
+
+#include "HLODBuilder.h"
+
+#include "Engine/StaticMesh.h"
+
+#include "WorldPartition/HLOD/HLODActor.h"
+#include "WorldPartition/HLOD/HLODLayer.h"
+
+DEFINE_LOG_CATEGORY(LogHLODBuilder);
+
+
+void FHLODBuilder::Build(AWorldPartitionHLOD* InHLODActor, const UHLODLayer* InHLODLayer, const TArray<FWorldPartitionReference>& InSubActors)
+{
+	TArray<UPrimitiveComponent*> SubComponents = GatherPrimitiveComponents(InSubActors);
+	if (SubComponents.IsEmpty())
+	{
+		return;
+	}
+
+	TArray<UPrimitiveComponent*> HLODPrimitives = CreateComponents(InHLODActor, InHLODLayer, SubComponents);
+	HLODPrimitives.RemoveSwap(nullptr);
+
+	if (!HLODPrimitives.IsEmpty())
+	{
+		InHLODActor->Modify();
+		InHLODActor->SetHLODPrimitives(HLODPrimitives);
+	}
+}
+
+TArray<UPrimitiveComponent*> FHLODBuilder::GatherPrimitiveComponents(const TArray<FWorldPartitionReference>& InActors)
+{
+	TArray<UPrimitiveComponent*> PrimitiveComponents;
+
+	auto GatherPrimitivesFromActor = [&PrimitiveComponents](const AActor* Actor, const AActor* ParentActor = nullptr)
+	{
+		const TCHAR* Padding = ParentActor ? TEXT("    ") : TEXT("");
+		UE_LOG(LogHLODBuilder, Verbose, TEXT("%s* Adding components from actor %s"), Padding, *Actor->GetName());
+		for (UActorComponent* SubComponent : Actor->GetComponents())
+		{
+			if (SubComponent && SubComponent->IsHLODRelevant())
+			{
+				if (UPrimitiveComponent* PrimitiveComponent = Cast<UPrimitiveComponent>(SubComponent))
+				{
+					PrimitiveComponents.Add(PrimitiveComponent);
+
+					if (UInstancedStaticMeshComponent* ISMC = Cast<UInstancedStaticMeshComponent>(PrimitiveComponent))
+					{
+						UE_LOG(LogHLODBuilder, Verbose, TEXT("%s    * %s [%d instances]"), Padding, *ISMC->GetStaticMesh()->GetName(), ISMC->GetInstanceCount());
+					}
+					else if (UStaticMeshComponent* SMC = Cast<UStaticMeshComponent>(PrimitiveComponent))
+					{
+						UE_LOG(LogHLODBuilder, Verbose, TEXT("%s    * %s"), Padding, *SMC->GetStaticMesh()->GetName());
+					}
+				}
+				else
+				{
+					UE_LOG(LogHLODBuilder, Warning, TEXT("Component \"%s\" is marked as HLOD-relevant but this type of component currently unsupported."), *SubComponent->GetFullName());
+				}
+			}
+		}
+	};
+
+	TSet<AActor*> UnderlyingActors;
+
+	for (const FWorldPartitionReference& ActorRef : InActors)
+	{
+		AActor* Actor = ActorRef->GetActor();
+
+		// Gather primitives from the Actor
+		GatherPrimitivesFromActor(Actor);
+
+		// Retrieve all underlying actors (ex: all sub actors of a LevelInstance)
+		UnderlyingActors.Reset();
+		Actor->EditorGetUnderlyingActors(UnderlyingActors);
+
+		// Gather primitives from underlying actors
+		for (const AActor* UnderlyingActor : UnderlyingActors)
+		{
+			if (UnderlyingActor->IsHLODRelevant())
+			{
+				GatherPrimitivesFromActor(UnderlyingActor, Actor);
+			}
+		}
+	}
+
+	return PrimitiveComponents;
+}
+
+void FHLODBuilder::DisableCollisions(UPrimitiveComponent* Component)
+{
+	Component->SetCollisionProfileName(UCollisionProfile::NoCollision_ProfileName);
+	Component->SetGenerateOverlapEvents(false);
+	Component->SetCanEverAffectNavigation(false);
+	Component->CanCharacterStepUpOn = ECanBeCharacterBase::ECB_No;
+	Component->SetCanEverAffectNavigation(false);
+	Component->SetCollisionEnabled(ECollisionEnabled::NoCollision);
+}
diff --git a/Engine/Source/Developer/WorldPartitionHLODUtilities/Private/HLODBuilder.h b/Engine/Source/Developer/WorldPartitionHLODUtilities/Private/HLODBuilder.h
new file mode 100644
--- /dev/null
+++ b/Engine/Source/Developer/WorldPartitionHLODUtilities/Private/HLODBuilder.h
@@ -0,0 +1,30 @@
+// Copyright Epic Games, Inc. All Rights Reserved.
+
+#pragma once
+
+#include "CoreMinimal.h"
+#include "WorldPartition/WorldPartitionHandle.h"
+
+class AWorldPartitionHLOD;
+class UPrimitiveComponent;
+class UHLODLayer;
+
+DECLARE_LOG_CATEGORY_EXTERN(LogHLODBuilder, Log, All);
+
+
+/**
+ * Base class for all HLODBuilders
+ */
+class FHLODBuilder
+{
+public:
+	virtual ~FHLODBuilder() {}
+
+	virtual TArray<UPrimitiveComponent*> CreateComponents(AWorldPartitionHLOD* InHLODActor, const UHLODLayer* InHLODLayer, const TArray<UPrimitiveComponent*>& InSubComponents) = 0;
+
+	void Build(AWorldPartitionHLOD* InHLODActor, const UHLODLayer* InHLODLayer, const TArray<FWorldPartitionReference>& InSubActors);
+
+	static TArray<UPrimitiveComponent*> GatherPrimitiveComponents(const TArray<FWorldPartitionReference>& InActors);
+
+	void DisableCollisions(UPrimitiveComponent* Component);
+};
diff --git a/Engine/Source/Developer/WorldPartitionHLODUtilities/Private/HLODBuilderInstancing.cpp b/Engine/Source/Developer/WorldPartitionHLODUtilities/Private/HLODBuilderInstancing.cpp
new file mode 100644
--- /dev/null
+++ b/Engine/Source/Developer/WorldPartitionHLODUtilities/Private/HLODBuilderInstancing.cpp
@@ -0,0 +1,116 @@
+// Copyright Epic Games, Inc. All Rights Reserved.
+
+#include "HLODBuilderInstancing.h"
+
+#include "Misc/HashBuilder.h"
+#include "Engine/StaticMesh.h"
+#include "WorldPartition/HLOD/HLODActor.h"
+#include "WorldPartition/HLOD/HLODLayer.h"
+
+
+// We want to merge all SMC that are using the same static mesh
+// However, we must also take material overiddes into account.
+struct FInstancingKey
+{
+	FInstancingKey(const UStaticMeshComponent* SMC)
+	{
+		FHashBuilder HashBuilder;
+
+		StaticMesh = SMC->GetStaticMesh();
+		HashBuilder << StaticMesh;
+
+		const int32 NumMaterials = SMC->GetNumMaterials();
+		Materials.Reserve(NumMaterials);
+
+		for (int32 MaterialIndex = 0; MaterialIndex < NumMaterials; ++MaterialIndex)
+		{
+			UMaterialInterface* Material = SMC->GetMaterial(MaterialIndex);
+
+			Materials.Add(Material);
+			HashBuilder << Material;
+		}
+
+		Hash = HashBuilder.GetHash();
+	}
+
+	void ApplyTo(UStaticMeshComponent* SMC) const
+	{
+		// Set static mesh
+		SMC->SetStaticMesh(StaticMesh);
+			
+		// Set material overrides
+		for (int32 MaterialIndex = 0; MaterialIndex < Materials.Num(); ++MaterialIndex)
+		{
+			SMC->SetMaterial(MaterialIndex, Materials[MaterialIndex]);
+		}
+	}
+
+	friend uint32 GetTypeHash(const FInstancingKey& Key)
+	{
+		return Key.Hash;
+	}
+
+	bool operator==(const FInstancingKey& Other) const
+	{
+		return Hash == Other.Hash && StaticMesh == Other.StaticMesh && Materials == Other.Materials;
+	}
+
+private:
+	UStaticMesh*				StaticMesh;
+	TArray<UMaterialInterface*>	Materials;
+	uint32						Hash;
+};
+
+
+TArray<UPrimitiveComponent*> FHLODBuilder_Instancing::CreateComponents(AWorldPartitionHLOD* InHLODActor, const UHLODLayer* InHLODLayer, const TArray<UPrimitiveComponent*>& InSubComponents)
+{
+	TRACE_CPUPROFILER_EVENT_SCOPE(FHLODBuilder_Instancing::CreateComponents);
+
+	TArray<UPrimitiveComponent*> Components;
+
+	// Gather all meshes to instantiate along with their transforms
+	TMap<FInstancingKey, TArray<UPrimitiveComponent*>> Instances;
+	for (UPrimitiveComponent* Primitive : InSubComponents)
+	{
+		if (UStaticMeshComponent* SMC = Cast<UStaticMeshComponent>(Primitive))
+		{
+			Instances.FindOrAdd(FInstancingKey(SMC)).Add(SMC);
+		}
+	}
+
+	// Create an ISMC for each SM asset we found
+	for (const auto& Entry : Instances)
+	{
+		const FInstancingKey EntryInstancingKey = Entry.Key;
+		const TArray<UPrimitiveComponent*>& EntryComponents = Entry.Value;
+			
+		UInstancedStaticMeshComponent* Component = NewObject<UInstancedStaticMeshComponent>(InHLODActor);
+		EntryInstancingKey.ApplyTo(Component);
+		Component->SetForcedLodModel(Component->GetStaticMesh()->GetNumLODs());
+
+		DisableCollisions(Component);
+
+		// Add all instances
+		for (UPrimitiveComponent* SMC : EntryComponents)
+		{
+			// If we have an ISMC, retrieve all instances
+			if (UInstancedStaticMeshComponent* InstancedStaticMeshComponent = Cast<UInstancedStaticMeshComponent>(SMC))
+			{
+				for (int32 InstanceIdx = 0; InstanceIdx < InstancedStaticMeshComponent->GetInstanceCount(); InstanceIdx++)
+				{
+					FTransform InstanceTransform;
+					InstancedStaticMeshComponent->GetInstanceTransform(InstanceIdx, InstanceTransform, true);
+					Component->AddInstanceWorldSpace(InstanceTransform);
+				}
+			}
+			else
+			{
+				Component->AddInstanceWorldSpace(SMC->GetComponentTransform());
+			}
+		}
+
+		Components.Add(Component);
+	};
+
+	return Components;
+}
diff --git a/Engine/Source/Developer/WorldPartitionHLODUtilities/Private/HLODBuilderInstancing.h b/Engine/Source/Developer/WorldPartitionHLODUtilities/Private/HLODBuilderInstancing.h
new file mode 100644
--- /dev/null
+++ b/Engine/Source/Developer/WorldPartitionHLODUtilities/Private/HLODBuilderInstancing.h
@@ -0,0 +1,13 @@
+// Copyright Epic Games, Inc. All Rights Reserved.
+
+#pragma once
+
+#include "HLODBuilder.h"
+
+/**
+ * Build a AWorldPartitionHLOD whose components are ISMC
+ */
+class FHLODBuilder_Instancing : public FHLODBuilder
+{
+	virtual TArray<UPrimitiveComponent*> CreateComponents(AWorldPartitionHLOD* InHLODActor, const UHLODLayer* InHLODLayer, const TArray<UPrimitiveComponent*>& InSubComponents) override;
+};
diff --git a/Engine/Source/Developer/WorldPartitionHLODUtilities/Private/HLODBuilderMeshApproximate.cpp b/Engine/Source/Developer/WorldPartitionHLODUtilities/Private/HLODBuilderMeshApproximate.cpp
new file mode 100644
--- /dev/null
+++ b/Engine/Source/Developer/WorldPartitionHLODUtilities/Private/HLODBuilderMeshApproximate.cpp
@@ -0,0 +1,88 @@
+// Copyright Epic Games, Inc. All Rights Reserved.
+
+#include "HLODBuilderMeshApproximate.h"
+
+#include "Algo/ForEach.h"
+#include "Algo/Transform.h"
+#include "WorldPartition/HLOD/HLODActor.h"
+#include "WorldPartition/HLOD/HLODLayer.h"
+
+#include "Engine/StaticMesh.h"
+#include "Materials/Material.h"
+
+#include "Modules/ModuleManager.h"
+#include "IGeometryProcessingInterfacesModule.h"
+#include "GeometryProcessingInterfaces/ApproximateActors.h"
+
+
+TArray<UPrimitiveComponent*> FHLODBuilder_MeshApproximate::CreateComponents(AWorldPartitionHLOD* InHLODActor, const UHLODLayer* InHLODLayer, const TArray<UPrimitiveComponent*>& InSubComponents)
+{
+	TRACE_CPUPROFILER_EVENT_SCOPE(FHLODBuilder_MeshSimplify::CreateComponents);
+
+	TSet<AActor*> Actors;
+	Algo::Transform(InSubComponents, Actors, [](UPrimitiveComponent* InPrimitiveComponent) { return InPrimitiveComponent->GetOwner(); });
+
+	IGeometryProcessingInterfacesModule& GeomProcInterfaces = FModuleManager::Get().LoadModuleChecked<IGeometryProcessingInterfacesModule>("GeometryProcessingInterfaces");
+	IGeometryProcessing_ApproximateActors* ApproxActorsAPI = GeomProcInterfaces.GetApproximateActorsImplementation();
+
+	//
+	// Construct options for ApproximateActors operation
+	//
+
+	const FMeshApproximationSettings& UseSettings = InHLODLayer->GetMeshApproximationSettings();
+
+	IGeometryProcessing_ApproximateActors::FOptions Options;
+	Options.BasePackagePath = InHLODActor->GetPackage()->GetName();
+	Options.WorldSpaceApproximationAccuracyMeters = UseSettings.ApproximationAccuracy;
+	Options.ClampVoxelDimension = UseSettings.ClampVoxelDimension;
+	Options.WindingThreshold = UseSettings.WindingThreshold;
+	Options.bApplyMorphology = UseSettings.bFillGaps;
+	Options.MorphologyDistanceMeters = UseSettings.GapDistance;
+	Options.FixedTriangleCount = UseSettings.TargetTriCount;
+	if (UseSettings.SimplifyMethod == EMeshApproximationSimplificationPolicy::TrianglesPerArea)
+	{
+		Options.MeshSimplificationPolicy = IGeometryProcessing_ApproximateActors::ESimplificationPolicy::TrianglesPerUnitSqMeter;
+		Options.SimplificationTargetMetric = UseSettings.TrianglesPerM;
+	}
+	else
+	{
+		Options.MeshSimplificationPolicy = IGeometryProcessing_ApproximateActors::ESimplificationPolicy::FixedTriangleCount;
+	}
+
+	Options.TextureImageSize = UseSettings.MaterialSettings.TextureSize.X;
+	Options.AntiAliasMultiSampling = FMath::Max(1, UseSettings.MultiSamplingAA);
+
+	Options.RenderCaptureImageSize = (UseSettings.RenderCaptureResolution == 0) ?
+		Options.TextureImageSize : UseSettings.RenderCaptureResolution;
+	Options.FieldOfViewDegrees = UseSettings.CaptureFieldOfView;
+	Options.NearPlaneDist = UseSettings.NearPlaneDist;
+
+	// run actor approximation computation
+	IGeometryProcessing_ApproximateActors::FResults Results;
+	ApproxActorsAPI->ApproximateActors(Actors.Array(), Options, Results);
+
+	auto FixupAsset = [&InHLODActor](UObject* Asset)
+	{
+		Asset->ClearFlags(RF_Public | RF_Standalone);
+		Asset->Rename(nullptr, InHLODActor->GetPackage(), REN_NonTransactional | REN_DontCreateRedirectors | REN_ForceNoResetLoaders);
+	};
+	
+	Algo::ForEach(Results.NewMeshAssets, FixupAsset);
+	Algo::ForEach(Results.NewMaterials, FixupAsset);
+	Algo::ForEach(Results.NewTextures, FixupAsset);
+
+	TArray<UPrimitiveComponent*> Components;
+	if (Results.ResultCode == IGeometryProcessing_ApproximateActors::EResultCode::Success)
+	{
+		for (UStaticMesh* StaticMesh : Results.NewMeshAssets)
+		{
+			UStaticMeshComponent* Component = NewObject<UStaticMeshComponent>(InHLODActor);
+			Component->SetStaticMesh(StaticMesh);
+			DisableCollisions(Component);
+
+			Components.Add(Component);
+		}		
+	}
+
+	return Components;
+}
diff --git a/Engine/Source/Developer/WorldPartitionHLODUtilities/Private/HLODBuilderMeshApproximate.h b/Engine/Source/Developer/WorldPartitionHLODUtilities/Private/HLODBuilderMeshApproximate.h
new file mode 100644
--- /dev/null
+++ b/Engine/Source/Developer/WorldPartitionHLODUtilities/Private/HLODBuilderMeshApproximate.h
@@ -0,0 +1,13 @@
+// Copyright Epic Games, Inc. All Rights Reserved.
+
+#pragma once
+
+#include "HLODBuilder.h"
+
+/**
+ * Build an approximated mesh using geometry from the provided actors
+ */
+class FHLODBuilder_MeshApproximate : public FHLODBuilder
+{
+	virtual TArray<UPrimitiveComponent*> CreateComponents(AWorldPartitionHLOD* InHLODActor, const UHLODLayer* InHLODLayer, const TArray<UPrimitiveComponent*>& InSubComponents) override;
+};
diff --git a/Engine/Source/Developer/WorldPartitionHLODUtilities/Private/HLODBuilderMeshMerge.cpp b/Engine/Source/Developer/WorldPartitionHLODUtilities/Private/HLODBuilderMeshMerge.cpp
new file mode 100644
--- /dev/null
+++ b/Engine/Source/Developer/WorldPartitionHLODUtilities/Private/HLODBuilderMeshMerge.cpp
@@ -0,0 +1,43 @@
+// Copyright Epic Games, Inc. All Rights Reserved.
+
+#include "HLODBuilderMeshMerge.h"
+
+#include "Algo/ForEach.h"
+#include "Engine/StaticMesh.h"
+#include "Materials/Material.h"
+#include "Components/StaticMeshComponent.h"
+
+#include "WorldPartition/HLOD/HLODActor.h"
+#include "WorldPartition/HLOD/HLODLayer.h"
+
+#include "IMeshMergeUtilities.h"
+#include "MeshMergeModule.h"
+#include "Modules/ModuleManager.h"
+
+
+TArray<UPrimitiveComponent*> FHLODBuilder_MeshMerge::CreateComponents(AWorldPartitionHLOD* InHLODActor, const UHLODLayer* InHLODLayer, const TArray<UPrimitiveComponent*>& InSubComponents)
+{
+	TRACE_CPUPROFILER_EVENT_SCOPE(FHLODBuilder_MeshMerge::CreateComponents);
+
+	TArray<UObject*> Assets;
+	FVector MergedActorLocation;
+
+	const IMeshMergeUtilities& MeshMergeUtilities = FModuleManager::Get().LoadModuleChecked<IMeshMergeModule>("MeshMergeUtilities").GetUtilities();
+	MeshMergeUtilities.MergeComponentsToStaticMesh(InSubComponents, InHLODActor->GetWorld(), InHLODLayer->GetMeshMergeSettings(), InHLODLayer->GetHLODMaterial().LoadSynchronous(), InHLODActor->GetPackage(), InHLODActor->GetActorLabel(), Assets, MergedActorLocation, 0.25f, false);
+
+	UStaticMeshComponent* Component = nullptr;
+	Algo::ForEach(Assets, [this, InHLODActor, &Component, &MergedActorLocation](UObject* Asset)
+	{
+		Asset->ClearFlags(RF_Public | RF_Standalone);
+
+		if (Cast<UStaticMesh>(Asset))
+		{
+			Component = NewObject<UStaticMeshComponent>(InHLODActor);
+			Component->SetStaticMesh(static_cast<UStaticMesh*>(Asset));
+			Component->SetWorldLocation(MergedActorLocation);
+			DisableCollisions(Component);
+		}
+	});
+
+	return TArray<UPrimitiveComponent*>({ Component });
+}
diff --git a/Engine/Source/Developer/WorldPartitionHLODUtilities/Private/HLODBuilderMeshMerge.h b/Engine/Source/Developer/WorldPartitionHLODUtilities/Private/HLODBuilderMeshMerge.h
new file mode 100644
--- /dev/null
+++ b/Engine/Source/Developer/WorldPartitionHLODUtilities/Private/HLODBuilderMeshMerge.h
@@ -0,0 +1,13 @@
+// Copyright Epic Games, Inc. All Rights Reserved.
+
+#pragma once
+
+#include "HLODBuilder.h"
+
+/**
+ * Build a merged mesh using geometry from the provided actors
+ */
+class FHLODBuilder_MeshMerge : public FHLODBuilder
+{
+	virtual TArray<UPrimitiveComponent*> CreateComponents(AWorldPartitionHLOD* InHLODActor, const UHLODLayer* InHLODLayer, const TArray<UPrimitiveComponent*>& InSubComponents) override;
+};
diff --git a/Engine/Source/Developer/WorldPartitionHLODUtilities/Private/HLODBuilderMeshSimplify.cpp b/Engine/Source/Developer/WorldPartitionHLODUtilities/Private/HLODBuilderMeshSimplify.cpp
new file mode 100644
--- /dev/null
+++ b/Engine/Source/Developer/WorldPartitionHLODUtilities/Private/HLODBuilderMeshSimplify.cpp
@@ -0,0 +1,47 @@
+// Copyright Epic Games, Inc. All Rights Reserved.
+
+#include "HLODBuilderMeshSimplify.h"
+
+#include "Algo/ForEach.h"
+#include "Algo/Transform.h"
+#include "Engine/StaticMesh.h"
+#include "Materials/Material.h"
+#include "Components/StaticMeshComponent.h"
+
+#include "WorldPartition/HLOD/HLODActor.h"
+#include "WorldPartition/HLOD/HLODLayer.h"
+
+#include "IMeshMergeUtilities.h"
+#include "MeshMergeModule.h"
+#include "Modules/ModuleManager.h"
+
+
+TArray<UPrimitiveComponent*> FHLODBuilder_MeshSimplify::CreateComponents(AWorldPartitionHLOD* InHLODActor, const UHLODLayer* InHLODLayer, const TArray<UPrimitiveComponent*>& InSubComponents)
+{
+	TRACE_CPUPROFILER_EVENT_SCOPE(FHLODBuilder_MeshSimplify::CreateComponents);
+
+	TArray<UObject*> Assets;
+	FCreateProxyDelegate ProxyDelegate;
+	ProxyDelegate.BindLambda([&Assets](const FGuid Guid, TArray<UObject*>& InAssetsCreated) { Assets = InAssetsCreated; });
+
+	TArray<UStaticMeshComponent*> StaticMeshComponents;
+	Algo::TransformIf(InSubComponents, StaticMeshComponents, [](UPrimitiveComponent* InPrimitiveComponent) { return InPrimitiveComponent->IsA<UStaticMeshComponent>(); }, [](UPrimitiveComponent* InPrimitiveComponent) { return Cast<UStaticMeshComponent>(InPrimitiveComponent); });
+
+	const IMeshMergeUtilities& MeshMergeUtilities = FModuleManager::Get().LoadModuleChecked<IMeshMergeModule>("MeshMergeUtilities").GetUtilities();
+	MeshMergeUtilities.CreateProxyMesh(StaticMeshComponents, InHLODLayer->GetMeshSimplifySettings(), InHLODLayer->GetHLODMaterial().LoadSynchronous(), InHLODActor->GetPackage(), InHLODActor->GetActorLabel(), FGuid::NewGuid(), ProxyDelegate, true);
+
+	UStaticMeshComponent* Component = nullptr;
+	Algo::ForEach(Assets, [this, InHLODActor, &Component](UObject* Asset)
+	{
+		Asset->ClearFlags(RF_Public | RF_Standalone);
+
+		if (Cast<UStaticMesh>(Asset))
+		{
+			Component = NewObject<UStaticMeshComponent>(InHLODActor);
+			Component->SetStaticMesh(static_cast<UStaticMesh*>(Asset));
+			DisableCollisions(Component);
+		}
+	});
+
+	return TArray<UPrimitiveComponent*>({ Component });
+}
diff --git a/Engine/Source/Developer/WorldPartitionHLODUtilities/Private/HLODBuilderMeshSimplify.h b/Engine/Source/Developer/WorldPartitionHLODUtilities/Private/HLODBuilderMeshSimplify.h
new file mode 100644
--- /dev/null
+++ b/Engine/Source/Developer/WorldPartitionHLODUtilities/Private/HLODBuilderMeshSimplify.h
@@ -0,0 +1,13 @@
+// Copyright Epic Games, Inc. All Rights Reserved.
+
+#pragma once
+
+#include "HLODBuilder.h"
+
+/**
+ * Build a simplified mesh using geometry from the provided actors
+ */
+class FHLODBuilder_MeshSimplify : public FHLODBuilder
+{
+	virtual TArray<UPrimitiveComponent*> CreateComponents(AWorldPartitionHLOD* InHLODActor, const UHLODLayer* InHLODLayer, const TArray<UPrimitiveComponent*>& InSubComponents) override;
+};
diff --git a/Engine/Source/Developer/WorldPartitionHLODUtilities/Private/WorldPartitionHLODUtilities.cpp b/Engine/Source/Developer/WorldPartitionHLODUtilities/Private/WorldPartitionHLODUtilities.cpp
--- a/Engine/Source/Developer/WorldPartitionHLODUtilities/Private/WorldPartitionHLODUtilities.cpp
+++ b/Engine/Source/Developer/WorldPartitionHLODUtilities/Private/WorldPartitionHLODUtilities.cpp
@@ -12,16 +12,9 @@
 
 #include "ISMPartition/ISMComponentDescriptor.h"
 
-#include "Algo/ForEach.h"
-#include "Algo/RemoveIf.h"
-#include "Algo/Transform.h"
-#include "Misc/HashBuilder.h"
 #include "Serialization/ArchiveCrc32.h"
 #include "Templates/UniquePtr.h"
 
-#include "IMeshMergeUtilities.h"
-#include "MeshMergeModule.h"
-#include "Modules/ModuleManager.h"
 #include "Engine/StaticMesh.h"
 #include "Engine/HLODProxy.h"
 #include "Materials/Material.h"
@@ -30,402 +23,125 @@
 #include "LevelInstance/LevelInstanceActor.h"
 #include "LevelInstance/LevelInstanceSubsystem.h"
 
-DEFINE_LOG_CATEGORY_STATIC(LogHLODBuilder, Log, All);
+#include "HLODBuilder.h"
+#include "HLODBuilderInstancing.h"
+#include "HLODBuilderMeshMerge.h"
+#include "HLODBuilderMeshSimplify.h"
+#include "HLODBuilderMeshApproximate.h"
 
-/**
- * Base class for all HLODBuilders
- */
-class FHLODBuilder
+static bool LoadSubActors(AWorldPartitionHLOD* InHLODActor, TArray<FWorldPartitionReference>& OutActors)
 {
-public:
-	virtual ~FHLODBuilder() {}
+	OutActors.Reserve(InHLODActor->GetSubActors().Num());
 
-	virtual TArray<UPrimitiveComponent*> CreateComponents(AWorldPartitionHLOD* InHLODActor, const UHLODLayer* InHLODLayer, const TArray<UPrimitiveComponent*>& InSubComponents) = 0;
+	UWorld* World = InHLODActor->GetWorld();
+	UWorldPartition* WorldPartition = World->GetWorldPartition();
+	check(WorldPartition);
 
-	void Build(AWorldPartitionHLOD* InHLODActor, const UHLODLayer* InHLODLayer, const TArray<const AActor*>& InSubActors)
-	{
-		TArray<UPrimitiveComponent*> SubComponents = GatherPrimitiveComponents(InSubActors);
-		if (SubComponents.IsEmpty())
-		{
-			return;
-		}
+	bool bIsDirty = false;
 
-		TArray<UPrimitiveComponent*> HLODPrimitives = CreateComponents(InHLODActor, InHLODLayer, SubComponents);
-		HLODPrimitives.RemoveSwap(nullptr);
-
-		if (!HLODPrimitives.IsEmpty())
-		{
-			InHLODActor->Modify();
-			InHLODActor->SetHLODPrimitives(HLODPrimitives);
-		}
-	}
-
-	static bool LoadSubActors(AWorldPartitionHLOD* InHLODActor, TArray<const AActor*>& OutActors, TArray<FWorldPartitionReference>& OutActorReferences)
+	// Gather (and potentially load) actors
+	for (const FGuid& SubActorGuid : InHLODActor->GetSubActors())
 	{
-		OutActors.Reserve(InHLODActor->GetSubActors().Num());
-		OutActorReferences.Reserve(InHLODActor->GetSubActors().Num());
-
-		UWorld* World = InHLODActor->GetWorld();
-		UWorldPartition* WorldPartition = World->GetWorldPartition();
-		check(WorldPartition);
+		FWorldPartitionReference ActorRef(WorldPartition, SubActorGuid);
 
-		bool bIsDirty = false;
-
-		// Gather (and potentially load) actors
-		for (const FGuid& SubActorGuid : InHLODActor->GetSubActors())
+		if (ActorRef.IsValid())
 		{
-			FWorldPartitionReference ActorRef(WorldPartition, SubActorGuid);
-
-			if (ActorRef.IsValid())
-			{
-				AActor* LoadedActor = ActorRef.Get()->GetActor();
-
-				// Load level instances
-				if (ALevelInstance* LevelInstance = Cast<ALevelInstance>(LoadedActor))
-				{
-					// Wait for level instance loading
-					if (LevelInstance->SupportsLoading())
-					{
-						LevelInstance->GetLevelInstanceSubsystem()->BlockLoadLevelInstance(LevelInstance);
-					}
-				}
+			AActor* LoadedActor = ActorRef.Get()->GetActor();
 
-				OutActors.Add(LoadedActor);
-				OutActorReferences.Add(MoveTemp(ActorRef));
-			}
-			else
+			// Load level instances
+			if (ALevelInstance* LevelInstance = Cast<ALevelInstance>(LoadedActor))
 			{
-				bIsDirty = true;
-			}
-		}
-
-		// Wait for compilation to finish
-		FStaticMeshCompilingManager::Get().FinishAllCompilation();
-
-		return !bIsDirty;
-	}
-
-	static TArray<UPrimitiveComponent*> GatherPrimitiveComponents(const TArray<const AActor*>& InActors)
-	{
-		TArray<UPrimitiveComponent*> PrimitiveComponents;
-
-		auto GatherPrimitivesFromActor = [&PrimitiveComponents](const AActor* Actor, const AActor* ParentActor = nullptr)
-		{
-			const TCHAR* Padding = ParentActor ? TEXT("    ") : TEXT("");
-			UE_LOG(LogHLODBuilder, Verbose, TEXT("%s* Adding components from actor %s"), Padding, *Actor->GetName());
-			for (UActorComponent* SubComponent : Actor->GetComponents())
-			{
-				if (SubComponent && SubComponent->IsHLODRelevant())
+				// Wait for level instance loading
+				if (LevelInstance->SupportsLoading())
 				{
-					if (UPrimitiveComponent* PrimitiveComponent = Cast<UPrimitiveComponent>(SubComponent))
-					{
-						PrimitiveComponents.Add(PrimitiveComponent);
-
-						if (UInstancedStaticMeshComponent* ISMC = Cast<UInstancedStaticMeshComponent>(PrimitiveComponent))
-						{
-							UE_LOG(LogHLODBuilder, Verbose, TEXT("%s    * %s [%d instances]"), Padding, *ISMC->GetStaticMesh()->GetName(), ISMC->GetInstanceCount());
-						}
-						else if (UStaticMeshComponent* SMC = Cast<UStaticMeshComponent>(PrimitiveComponent))
-						{
-							UE_LOG(LogHLODBuilder, Verbose, TEXT("%s    * %s"), Padding, *SMC->GetStaticMesh()->GetName());
-						}
-					}
-					else
-					{
-						UE_LOG(LogHLODBuilder, Warning, TEXT("Component \"%s\" is marked as HLOD-relevant but this type of component currently unsupported."), *SubComponent->GetFullName());
-					}
+					LevelInstance->GetLevelInstanceSubsystem()->BlockLoadLevelInstance(LevelInstance);
 				}
 			}
-		};
-
-		TSet<AActor*> UnderlyingActors;
-
-		for (const AActor* Actor : InActors)
-		{
-			// Gather primitives from the Actor
-			GatherPrimitivesFromActor(Actor);
-
-			// Retrieve all underlying actors (ex: all sub actors of a LevelInstance)
-			UnderlyingActors.Reset();
-			Actor->EditorGetUnderlyingActors(UnderlyingActors);
 
-			// Gather primitives from underlying actors
-			for (const AActor* UnderlyingActor : UnderlyingActors)
-			{
-				if (UnderlyingActor->IsHLODRelevant())
-				{
-					GatherPrimitivesFromActor(UnderlyingActor, Actor);
-				}
-			}
+			OutActors.Add(MoveTemp(ActorRef));
 		}
-
-		return PrimitiveComponents;
-	}
-
-	static uint32 ComputeHLODHash(AWorldPartitionHLOD* InHLODActor, const TArray<const AActor*>& InActors)
-	{
-		FArchiveCrc32 Ar;
-
-		// Base key, changing this will force a rebuild of all HLODs
-		FString HLODBaseKey = "5184EABE2BEF440DB5A461554A28A3E4";
-		Ar << HLODBaseKey;
-
-		// HLOD Layer
-		uint32 HLODLayerHash = InHLODActor->GetSubActorsHLODLayer()->GetCRC();
-		UE_LOG(LogHLODBuilder, VeryVerbose, TEXT(" - HLODLayer (%s) = %x"), *InHLODActor->GetSubActorsHLODLayer()->GetName(), HLODLayerHash);
-		Ar << HLODLayerHash;
-
-		// We get the CRC of each component
-		TArray<uint32> ComponentsCRCs;
-		for (UPrimitiveComponent* Component : GatherPrimitiveComponents(InActors))
+		else
 		{
-			if (UStaticMeshComponent* StaticMeshComponent = Cast<UStaticMeshComponent>(Component))
-			{
-				uint32 ComponentCRC = 0;
-
-				UE_LOG(LogHLODBuilder, VeryVerbose, TEXT(" - Component \'%s\' from actor \'%s\'"), *Component->GetName(), *Component->GetOwner()->GetName());
-
-				// CRC component
-				uint32 StaticMeshComponentCRC = UHLODProxy::GetCRC(StaticMeshComponent);
-				UE_LOG(LogHLODBuilder, VeryVerbose, TEXT("     - StaticMeshComponent (%s) = %x"), *StaticMeshComponent->GetName(), StaticMeshComponentCRC);
-				ComponentCRC = HashCombine(ComponentCRC, StaticMeshComponentCRC);
-
-				if (UStaticMesh* StaticMesh = StaticMeshComponent->GetStaticMesh())
-				{
-					// CRC static mesh
-					int32 StaticMeshCRC = UHLODProxy::GetCRC(StaticMesh);
-					UE_LOG(LogHLODBuilder, VeryVerbose, TEXT("     - StaticMesh (%s) = %x"), *StaticMesh->GetName(), StaticMeshCRC);
-					ComponentCRC = HashCombine(ComponentCRC, StaticMeshCRC);
-
-					// CRC materials
-					const int32 NumMaterials = StaticMeshComponent->GetNumMaterials();
-					for (int32 MaterialIndex = 0; MaterialIndex < NumMaterials; ++MaterialIndex)
-					{
-						UMaterialInterface* MaterialInterface = StaticMeshComponent->GetMaterial(MaterialIndex);
-						if (MaterialInterface)
-						{
-							uint32 MaterialInterfaceCRC = UHLODProxy::GetCRC(MaterialInterface);
-							UE_LOG(LogHLODBuilder, VeryVerbose, TEXT("     - MaterialInterface (%s) = %x"), *MaterialInterface->GetName(), MaterialInterfaceCRC);
-							ComponentCRC = HashCombine(ComponentCRC, MaterialInterfaceCRC);
-
-							TArray<UTexture*> Textures;
-							MaterialInterface->GetUsedTextures(Textures, EMaterialQualityLevel::High, true, ERHIFeatureLevel::SM5, true);
-							for (UTexture* Texture : Textures)
-							{
-								uint32 TextureCRC = UHLODProxy::GetCRC(Texture);
-								UE_LOG(LogHLODBuilder, VeryVerbose, TEXT("     - Texture (%s) = %x"), *Texture->GetName(), TextureCRC);
-								ComponentCRC = HashCombine(ComponentCRC, TextureCRC);
-							}
-						}
-					}
-				}
-
-				ComponentsCRCs.Add(ComponentCRC);
-			}
+			bIsDirty = true;
 		}
-
-		// Sort the components CRCs to ensure the order of components won't have an impact on the final CRC
-		ComponentsCRCs.Sort();
-
-		// Append all components CRCs
-		Ar << ComponentsCRCs;
-			
-		return Ar.GetCrc();
 	}
 
-	void DisableCollisions(UPrimitiveComponent* Component)
-	{
-		Component->SetCollisionProfileName(UCollisionProfile::NoCollision_ProfileName);
-		Component->SetGenerateOverlapEvents(false);
-		Component->SetCanEverAffectNavigation(false);
-		Component->CanCharacterStepUpOn = ECanBeCharacterBase::ECB_No;
-		Component->SetCanEverAffectNavigation(false);
-		Component->SetCollisionEnabled(ECollisionEnabled::NoCollision);
-	}
-};
+	// Wait for compilation to finish
+	FStaticMeshCompilingManager::Get().FinishAllCompilation();
 
+	return !bIsDirty;
+}
 
-/**
- * Build a AWorldPartitionHLOD whose components are ISMC
- */
-class FHLODBuilder_Instancing : public FHLODBuilder
+static uint32 ComputeHLODHash(AWorldPartitionHLOD* InHLODActor, const TArray<FWorldPartitionReference>& InActors)
 {
-	// We want to merge all SMC that are using the same static mesh
-	// However, we must also take material overiddes into account.
-	struct FInstancingKey
-	{
-		FInstancingKey(const UStaticMeshComponent* SMC)
-		{
-			FHashBuilder HashBuilder;
-
-			StaticMesh = SMC->GetStaticMesh();
-			HashBuilder << StaticMesh;
-
-			const int32 NumMaterials = SMC->GetNumMaterials();
-			Materials.Reserve(NumMaterials);
-
-			for (int32 MaterialIndex = 0; MaterialIndex < NumMaterials; ++MaterialIndex)
-			{
-				UMaterialInterface* Material = SMC->GetMaterial(MaterialIndex);
-
-				Materials.Add(Material);
-				HashBuilder << Material;
-			}
-
-			Hash = HashBuilder.GetHash();
-		}
-
-		void ApplyTo(UStaticMeshComponent* SMC) const
-		{
-			// Set static mesh
-			SMC->SetStaticMesh(StaticMesh);
-			
-			// Set material overrides
-			for (int32 MaterialIndex = 0; MaterialIndex < Materials.Num(); ++MaterialIndex)
-			{
-				SMC->SetMaterial(MaterialIndex, Materials[MaterialIndex]);
-			}
-		}
+	FArchiveCrc32 Ar;
 
-		friend uint32 GetTypeHash(const FInstancingKey& Key)
-		{
-			return Key.Hash;
-		}
+	// Base key, changing this will force a rebuild of all HLODs
+	FString HLODBaseKey = "5184EABE2BEF440DB5A461554A28A3E4";
+	Ar << HLODBaseKey;
 
-		bool operator==(const FInstancingKey& Other) const
-		{
-			return Hash == Other.Hash && StaticMesh == Other.StaticMesh && Materials == Other.Materials;
-		}
+	// HLOD Layer
+	uint32 HLODLayerHash = InHLODActor->GetSubActorsHLODLayer()->GetCRC();
+	UE_LOG(LogHLODBuilder, VeryVerbose, TEXT(" - HLODLayer (%s) = %x"), *InHLODActor->GetSubActorsHLODLayer()->GetName(), HLODLayerHash);
+	Ar << HLODLayerHash;
 
-	private:
-		UStaticMesh*				StaticMesh;
-		TArray<UMaterialInterface*>	Materials;
-		uint32						Hash;
-	};
-
-	virtual TArray<UPrimitiveComponent*> CreateComponents(AWorldPartitionHLOD* InHLODActor, const UHLODLayer* InHLODLayer, const TArray<UPrimitiveComponent*>& InSubComponents) override
+	// We get the CRC of each component
+	TArray<uint32> ComponentsCRCs;
+	for (UPrimitiveComponent* Component : FHLODBuilder::GatherPrimitiveComponents(InActors))
 	{
-		TRACE_CPUPROFILER_EVENT_SCOPE(FHLODBuilder_Instancing::CreateComponents);
-
-		TArray<UPrimitiveComponent*> Components;
-
-		// Gather all meshes to instantiate along with their transforms
-		TMap<FInstancingKey, TArray<UPrimitiveComponent*>> Instances;
-		for (UPrimitiveComponent* Primitive : InSubComponents)
+		if (UStaticMeshComponent* StaticMeshComponent = Cast<UStaticMeshComponent>(Component))
 		{
-			if (UStaticMeshComponent* SMC = Cast<UStaticMeshComponent>(Primitive))
-			{
-				Instances.FindOrAdd(FInstancingKey(SMC)).Add(SMC);
-			}
-		}
+			uint32 ComponentCRC = 0;
 
-		// Create an ISMC for each SM asset we found
-		for (const auto& Entry : Instances)
-		{
-			const FInstancingKey EntryInstancingKey = Entry.Key;
-			const TArray<UPrimitiveComponent*>& EntryComponents = Entry.Value;
-			
-			UInstancedStaticMeshComponent* Component = NewObject<UInstancedStaticMeshComponent>(InHLODActor);
-			EntryInstancingKey.ApplyTo(Component);
-			Component->SetForcedLodModel(Component->GetStaticMesh()->GetNumLODs());
+			UE_LOG(LogHLODBuilder, VeryVerbose, TEXT(" - Component \'%s\' from actor \'%s\'"), *Component->GetName(), *Component->GetOwner()->GetName());
 
-			DisableCollisions(Component);
+			// CRC component
+			uint32 StaticMeshComponentCRC = UHLODProxy::GetCRC(StaticMeshComponent);
+			UE_LOG(LogHLODBuilder, VeryVerbose, TEXT("     - StaticMeshComponent (%s) = %x"), *StaticMeshComponent->GetName(), StaticMeshComponentCRC);
+			ComponentCRC = HashCombine(ComponentCRC, StaticMeshComponentCRC);
 
-			// Add all instances
-			for (UPrimitiveComponent* SMC : EntryComponents)
+			if (UStaticMesh* StaticMesh = StaticMeshComponent->GetStaticMesh())
 			{
-				// If we have an ISMC, retrieve all instances
-				if (UInstancedStaticMeshComponent* InstancedStaticMeshComponent = Cast<UInstancedStaticMeshComponent>(SMC))
+				// CRC static mesh
+				int32 StaticMeshCRC = UHLODProxy::GetCRC(StaticMesh);
+				UE_LOG(LogHLODBuilder, VeryVerbose, TEXT("     - StaticMesh (%s) = %x"), *StaticMesh->GetName(), StaticMeshCRC);
+				ComponentCRC = HashCombine(ComponentCRC, StaticMeshCRC);
+
+				// CRC materials
+				const int32 NumMaterials = StaticMeshComponent->GetNumMaterials();
+				for (int32 MaterialIndex = 0; MaterialIndex < NumMaterials; ++MaterialIndex)
 				{
-					for (int32 InstanceIdx = 0; InstanceIdx < InstancedStaticMeshComponent->GetInstanceCount(); InstanceIdx++)
+					UMaterialInterface* MaterialInterface = StaticMeshComponent->GetMaterial(MaterialIndex);
+					if (MaterialInterface)
 					{
-						FTransform InstanceTransform;
-						InstancedStaticMeshComponent->GetInstanceTransform(InstanceIdx, InstanceTransform, true);
-						Component->AddInstanceWorldSpace(InstanceTransform);
+						uint32 MaterialInterfaceCRC = UHLODProxy::GetCRC(MaterialInterface);
+						UE_LOG(LogHLODBuilder, VeryVerbose, TEXT("     - MaterialInterface (%s) = %x"), *MaterialInterface->GetName(), MaterialInterfaceCRC);
+						ComponentCRC = HashCombine(ComponentCRC, MaterialInterfaceCRC);
+
+						TArray<UTexture*> Textures;
+						MaterialInterface->GetUsedTextures(Textures, EMaterialQualityLevel::High, true, ERHIFeatureLevel::SM5, true);
+						for (UTexture* Texture : Textures)
+						{
+							uint32 TextureCRC = UHLODProxy::GetCRC(Texture);
+							UE_LOG(LogHLODBuilder, VeryVerbose, TEXT("     - Texture (%s) = %x"), *Texture->GetName(), TextureCRC);
+							ComponentCRC = HashCombine(ComponentCRC, TextureCRC);
+						}
 					}
 				}
-				else
-				{
-					Component->AddInstanceWorldSpace(SMC->GetComponentTransform());
-				}
 			}
 
-			Components.Add(Component);
-		};
-
-		return Components;
-	}
-};
-
-
-/**
- * Build a merged mesh using geometry from the provided actors
- */
-class FHLODBuilder_MeshMerge : public FHLODBuilder
-{
-	virtual TArray<UPrimitiveComponent*> CreateComponents(AWorldPartitionHLOD* InHLODActor, const UHLODLayer* InHLODLayer, const TArray<UPrimitiveComponent*>& InSubComponents) override
-	{
-		TRACE_CPUPROFILER_EVENT_SCOPE(FHLODBuilder_MeshMerge::CreateComponents);
-
-		TArray<UObject*> Assets;
-		FVector MergedActorLocation;
-
-		const IMeshMergeUtilities& MeshMergeUtilities = FModuleManager::Get().LoadModuleChecked<IMeshMergeModule>("MeshMergeUtilities").GetUtilities();
-		MeshMergeUtilities.MergeComponentsToStaticMesh(InSubComponents, InHLODActor->GetWorld(), InHLODLayer->GetMeshMergeSettings(), InHLODLayer->GetHLODMaterial().LoadSynchronous(), InHLODActor->GetPackage(), InHLODActor->GetActorLabel(), Assets, MergedActorLocation, 0.25f, false);
-
-		UStaticMeshComponent* Component = nullptr;
-		Algo::ForEach(Assets, [this, InHLODActor, &Component, &MergedActorLocation](UObject* Asset)
-		{
-			Asset->ClearFlags(RF_Public | RF_Standalone);
-
-			if (Cast<UStaticMesh>(Asset))
-			{
-				Component = NewObject<UStaticMeshComponent>(InHLODActor);
-				Component->SetStaticMesh(static_cast<UStaticMesh*>(Asset));
-				Component->SetWorldLocation(MergedActorLocation);
-				DisableCollisions(Component);
-			}
-		});
-
-		return TArray<UPrimitiveComponent*>({ Component });
+			ComponentsCRCs.Add(ComponentCRC);
+		}
 	}
-};
-
-/**
- * Build a simplified mesh using geometry from the provided actors
- */
-class FHLODBuilder_MeshSimplify : public FHLODBuilder
-{
-	virtual TArray<UPrimitiveComponent*> CreateComponents(AWorldPartitionHLOD* InHLODActor, const UHLODLayer* InHLODLayer, const TArray<UPrimitiveComponent*>& InSubComponents) override
-	{
-		TRACE_CPUPROFILER_EVENT_SCOPE(FHLODBuilder_MeshSimplify::CreateComponents);
-
-		TArray<UObject*> Assets;
-		FCreateProxyDelegate ProxyDelegate;
-		ProxyDelegate.BindLambda([&Assets](const FGuid Guid, TArray<UObject*>& InAssetsCreated) { Assets = InAssetsCreated; });
-
-		TArray<UStaticMeshComponent*> StaticMeshComponents;
-		Algo::TransformIf(InSubComponents, StaticMeshComponents, [](UPrimitiveComponent* InPrimitiveComponent) { return InPrimitiveComponent->IsA<UStaticMeshComponent>(); }, [](UPrimitiveComponent* InPrimitiveComponent) { return Cast<UStaticMeshComponent>(InPrimitiveComponent); });
 
-		const IMeshMergeUtilities& MeshMergeUtilities = FModuleManager::Get().LoadModuleChecked<IMeshMergeModule>("MeshMergeUtilities").GetUtilities();
-		MeshMergeUtilities.CreateProxyMesh(StaticMeshComponents, InHLODLayer->GetMeshSimplifySettings(), InHLODLayer->GetHLODMaterial().LoadSynchronous(), InHLODActor->GetPackage(), InHLODActor->GetActorLabel(), FGuid::NewGuid(), ProxyDelegate, true);
+	// Sort the components CRCs to ensure the order of components won't have an impact on the final CRC
+	ComponentsCRCs.Sort();
 
-		UStaticMeshComponent* Component = nullptr;
-		Algo::ForEach(Assets, [this, InHLODActor, &Component](UObject* Asset)
-		{
-			Asset->ClearFlags(RF_Public | RF_Standalone);
-
-			if (Cast<UStaticMesh>(Asset))
-			{
-				Component = NewObject<UStaticMeshComponent>(InHLODActor);
-				Component->SetStaticMesh(static_cast<UStaticMesh*>(Asset));
-				DisableCollisions(Component);
-			}
-		});
+	// Append all components CRCs
+	Ar << ComponentsCRCs;
 
-		return TArray<UPrimitiveComponent*>({ Component });
-	}
-};
+	return Ar.GetCrc();
+}
 
 TArray<AWorldPartitionHLOD*> FWorldPartitionHLODUtilities::CreateHLODActors(FHLODCreationContext& InCreationContext, const FHLODCreationParams& InCreationParams, const TSet<FActorInstance>& InActors, const TArray<const UDataLayer*>& InDataLayers)
 {
@@ -563,17 +279,16 @@ TArray<AWorldPartitionHLOD*> FWorldPartitionHLODUtilities::CreateHLODActors(FHLO
 
 uint32 FWorldPartitionHLODUtilities::BuildHLOD(AWorldPartitionHLOD* InHLODActor)
 {
-	TArray<FWorldPartitionReference> ActorReferences;
-	TArray<const AActor*> SubActors;
+	TArray<FWorldPartitionReference> SubActors;
 	
-	bool bIsDirty = !FHLODBuilder::LoadSubActors(InHLODActor, SubActors, ActorReferences);
+	bool bIsDirty = !LoadSubActors(InHLODActor, SubActors);
 	if (bIsDirty)
 	{
 		UE_LOG(LogHLODBuilder, Warning, TEXT("HLOD actor \"%s\" needs to be rebuilt as it references actors that have been deleted."), *InHLODActor->GetActorLabel());
 	}
 
 	uint32 OldHLODHash = bIsDirty ? 0 : InHLODActor->GetHLODHash();
-	uint32 NewHLODHash = FHLODBuilder::ComputeHLODHash(InHLODActor, SubActors);
+	uint32 NewHLODHash = ComputeHLODHash(InHLODActor, SubActors);
 
 	if (OldHLODHash == NewHLODHash)
 	{
@@ -598,6 +313,10 @@ uint32 FWorldPartitionHLODUtilities::BuildHLOD(AWorldPartitionHLOD* InHLODActor)
 		HLODBuilder = TUniquePtr<FHLODBuilder>(new FHLODBuilder_MeshSimplify());
 		break;
 
+	case EHLODLayerType::MeshApproximate:
+		HLODBuilder = TUniquePtr<FHLODBuilder>(new FHLODBuilder_MeshApproximate());
+		break;
+
 	default:
 		checkf(false, TEXT("Unsupported type"));
 	}
diff --git a/Engine/Source/Developer/WorldPartitionHLODUtilities/WorldPartitionHLODUtilities.Build.cs b/Engine/Source/Developer/WorldPartitionHLODUtilities/WorldPartitionHLODUtilities.Build.cs
--- a/Engine/Source/Developer/WorldPartitionHLODUtilities/WorldPartitionHLODUtilities.Build.cs
+++ b/Engine/Source/Developer/WorldPartitionHLODUtilities/WorldPartitionHLODUtilities.Build.cs
@@ -27,10 +27,11 @@ public class WorldPartitionHLODUtilities : ModuleRules
 			}
         );
 
-        PrivateIncludePaths.AddRange(
+        PrivateIncludePathModuleNames.AddRange(
             new string[]
             {
-            }
+				"GeometryProcessingInterfaces",
+			}
         );
 
         DynamicallyLoadedModuleNames.AddRange(
@@ -39,7 +40,8 @@ public class WorldPartitionHLODUtilities : ModuleRules
                 "MeshUtilities",
                 "MeshMergeUtilities",
                 "MeshReductionInterface",
-            }
+				"GeometryProcessingInterfaces",
+			}
         );
 	}
 }
diff --git a/Engine/Source/Editor/MergeActors/Private/MeshApproximationTool/MeshApproximationTool.cpp b/Engine/Source/Editor/MergeActors/Private/MeshApproximationTool/MeshApproximationTool.cpp
--- a/Engine/Source/Editor/MergeActors/Private/MeshApproximationTool/MeshApproximationTool.cpp
+++ b/Engine/Source/Editor/MergeActors/Private/MeshApproximationTool/MeshApproximationTool.cpp
@@ -157,7 +157,7 @@ bool FMeshApproximationTool::RunMerge(const FString& PackageName, const TArray<T
 		return false;
 	}
 
-	const FMeshApproximationToolSettings& UseSettings = SettingsObject->Settings;
+	const FMeshApproximationSettings& UseSettings = SettingsObject->Settings;
 
 	IGeometryProcessingInterfacesModule& GeomProcInterfaces = FModuleManager::Get().LoadModuleChecked<IGeometryProcessingInterfacesModule>("GeometryProcessingInterfaces");
 	IGeometryProcessing_ApproximateActors* ApproxActorsAPI = GeomProcInterfaces.GetApproximateActorsImplementation();
diff --git a/Engine/Source/Editor/MergeActors/Private/MeshApproximationTool/MeshApproximationTool.h b/Engine/Source/Editor/MergeActors/Private/MeshApproximationTool/MeshApproximationTool.h
--- a/Engine/Source/Editor/MergeActors/Private/MeshApproximationTool/MeshApproximationTool.h
+++ b/Engine/Source/Editor/MergeActors/Private/MeshApproximationTool/MeshApproximationTool.h
@@ -10,105 +10,6 @@
 
 class SMeshApproximationDialog;
 
-
-
-UENUM()
-enum class EMeshApproximationSimplificationPolicy : uint8
-{
-	FixedTriangleCount = 0,
-	TrianglesPerArea = 1
-};
-
-
-USTRUCT(Blueprintable)
-struct FMeshApproximationToolSettings
-{
-	GENERATED_BODY()
-
-	/** Approximation Accuracy in Meters, will determine (eg) voxel resolution */
-	UPROPERTY(BlueprintReadWrite, EditAnywhere, Category = ShapeSettings, meta = (ClampMin = "0.1"))
-	float ApproximationAccuracy = 1.0;
-
-	/** Maximum allowable voxel count along main directions. This is a limit on ApproximationAccuracy */
-	UPROPERTY(BlueprintReadWrite, EditAnywhere, AdvancedDisplay, Category = ShapeSettings, meta = (ClampMin = "64"))
-	int32 ClampVoxelDimension = 512;
-
-	/** */
-	UPROPERTY(BlueprintReadWrite, EditAnywhere, Category = ShapeSettings, meta = (ClampMin = "0.01", ClampMax = "0.99"))
-	float WindingThreshold = 0.5;
-
-	UPROPERTY(BlueprintReadWrite, EditAnywhere, Category = ShapeSettings)
-	bool bFillGaps = true;
-
-	/** Distance in Meters to expand approximation to fill gaps */
-	UPROPERTY(BlueprintReadWrite, EditAnywhere, Category = ShapeSettings, meta = (ClampMin = "0.1", EditCondition = "bFillGaps"))
-	float GapDistance = 0.1;
-
-	UPROPERTY(BlueprintReadWrite, EditAnywhere, Category = ShapeSettings)
-	EMeshApproximationSimplificationPolicy SimplifyMethod = EMeshApproximationSimplificationPolicy::FixedTriangleCount;
-
-	UPROPERTY(BlueprintReadWrite, EditAnywhere, Category = ShapeSettings, meta = (ClampMin = "16", EditCondition = "SimplifyMethod == EMeshApproximationSimplificationPolicy::FixedTriangleCount" ))
-	int32 TargetTriCount = 2000;
-
-	/** Approximate Number of triangles per Square Meter */
-	UPROPERTY(BlueprintReadWrite, EditAnywhere, Category = ShapeSettings, meta = (ClampMin = "0.01", EditCondition = "SimplifyMethod == EMeshApproximationSimplificationPolicy::TrianglesPerArea" ))
-	float TrianglesPerM = 2.0f;
-
-
-
-	/** Material simplification */
-	UPROPERTY(BlueprintReadWrite, EditAnywhere, Category = MaterialSettings)
-	FMaterialProxySettings MaterialSettings;
-
-	/** If Value is > 1, Multisample output baked textures by this amount in each direction (eg 4 == 16x supersampling) */
-	UPROPERTY(BlueprintReadWrite, EditAnywhere, AdvancedDisplay, Category = MaterialSettings, meta = (ClampMin = "0", ClampMax = "8", UIMin = "0", UIMax = "4"))
-	int32 MultiSamplingAA = 0;
-
-	UPROPERTY(BlueprintReadWrite, EditAnywhere, AdvancedDisplay, Category = MaterialSettings, meta = (ClampMin = "5.0", ClampMax = "160.0"))
-	float CaptureFieldOfView = 30.0f;
-
-	UPROPERTY(BlueprintReadWrite, EditAnywhere, AdvancedDisplay, Category = MaterialSettings, meta = (ClampMin = "0.001", ClampMax = "1000.0"))
-	float NearPlaneDist = 1.0f;
-
-	/** If Value is zero, use MaterialSettings resolution, otherwise override the render capture resolution */
-	UPROPERTY(BlueprintReadWrite, EditAnywhere, AdvancedDisplay, Category = MaterialSettings, meta = (ClampMin = "128"))
-	int32 RenderCaptureResolution = 0;
-
-	/** Equality operator. */
-	bool operator==(const FMeshApproximationToolSettings& Other) const
-	{
-		return ApproximationAccuracy == Other.ApproximationAccuracy
-			&& ClampVoxelDimension == Other.ClampVoxelDimension
-			&& WindingThreshold == Other.WindingThreshold
-			&& bFillGaps == Other.bFillGaps
-			&& GapDistance == Other.GapDistance
-			&& SimplifyMethod == Other.SimplifyMethod
-			&& TargetTriCount == Other.TargetTriCount
-			&& TrianglesPerM == Other.TrianglesPerM
-
-			&& MaterialSettings == Other.MaterialSettings
-			&& MultiSamplingAA == Other.MultiSamplingAA
-			&& CaptureFieldOfView == Other.CaptureFieldOfView
-			&& NearPlaneDist == Other.NearPlaneDist
-			&& RenderCaptureResolution == Other.RenderCaptureResolution;
-	}
-
-	/** Inequality. */
-	bool operator!=(const FMeshApproximationToolSettings& Other) const
-	{
-		return !(*this == Other);
-	}
-
-#if WITH_EDITORONLY_DATA
-	/** Handles deprecated properties */
-	void PostLoadDeprecated() {}		// none currently
-#endif
-};
-
-
-
-
-
 /** Singleton wrapper to allow for using the setting structure in SSettingsView */
 UCLASS(config = Engine)
 class UMeshApproximationSettingsObject : public UObject
@@ -152,7 +53,7 @@ protected:
 	static UMeshApproximationSettingsObject* DefaultSettings;
 public:
 	UPROPERTY(EditAnywhere, meta = (ShowOnlyInnerProperties), Category = ApproximationSettings)
-	FMeshApproximationToolSettings Settings;
+	FMeshApproximationSettings Settings;
 };
 
 /**
diff --git a/Engine/Source/Runtime/Engine/Classes/Engine/MeshMerging.h b/Engine/Source/Runtime/Engine/Classes/Engine/MeshMerging.h
--- a/Engine/Source/Runtime/Engine/Classes/Engine/MeshMerging.h
+++ b/Engine/Source/Runtime/Engine/Classes/Engine/MeshMerging.h
@@ -677,4 +677,97 @@ struct FMeshInstancingSettings
 	 */
 	UPROPERTY(BlueprintReadWrite, EditAnywhere, Category = "Instancing", meta = (DisplayName = "Select the type of Instanced Component", DisallowedClasses = "FoliageInstancedStaticMeshComponent"))
 	TSubclassOf<UInstancedStaticMeshComponent> ISMComponentToUse;
+};
+
+UENUM()
+enum class EMeshApproximationSimplificationPolicy : uint8
+{
+	FixedTriangleCount = 0,
+	TrianglesPerArea = 1
+};
+
+
+USTRUCT(Blueprintable)
+struct FMeshApproximationSettings
+{
+	GENERATED_BODY()
+
+	/** Approximation Accuracy in Meters, will determine (eg) voxel resolution */
+	UPROPERTY(BlueprintReadWrite, EditAnywhere, Category = ShapeSettings, meta = (ClampMin = "0.1"))
+	float ApproximationAccuracy = 1.0;
+
+	/** Maximum allowable voxel count along main directions. This is a limit on ApproximationAccuracy */
+	UPROPERTY(BlueprintReadWrite, EditAnywhere, AdvancedDisplay, Category = ShapeSettings, meta = (ClampMin = "64"))
+	int32 ClampVoxelDimension = 512;
+
+	/** */
+	UPROPERTY(BlueprintReadWrite, EditAnywhere, Category = ShapeSettings, meta = (ClampMin = "0.01", ClampMax = "0.99"))
+	float WindingThreshold = 0.5;
+
+	UPROPERTY(BlueprintReadWrite, EditAnywhere, Category = ShapeSettings)
+	bool bFillGaps = true;
+
+	/** Distance in Meters to expand approximation to fill gaps */
+	UPROPERTY(BlueprintReadWrite, EditAnywhere, Category = ShapeSettings, meta = (ClampMin = "0.1", EditCondition = "bFillGaps"))
+	float GapDistance = 0.1;
+
+	UPROPERTY(BlueprintReadWrite, EditAnywhere, Category = ShapeSettings)
+	EMeshApproximationSimplificationPolicy SimplifyMethod = EMeshApproximationSimplificationPolicy::FixedTriangleCount;
+
+	UPROPERTY(BlueprintReadWrite, EditAnywhere, Category = ShapeSettings, meta = (ClampMin = "16", EditCondition = "SimplifyMethod == EMeshApproximationSimplificationPolicy::FixedTriangleCount" ))
+	int32 TargetTriCount = 2000;
+
+	/** Approximate Number of triangles per Square Meter */
+	UPROPERTY(BlueprintReadWrite, EditAnywhere, Category = ShapeSettings, meta = (ClampMin = "0.01", EditCondition = "SimplifyMethod == EMeshApproximationSimplificationPolicy::TrianglesPerArea" ))
+	float TrianglesPerM = 2.0f;
+
+
+
+	/** Material simplification */
+	UPROPERTY(BlueprintReadWrite, EditAnywhere, Category = MaterialSettings)
+	FMaterialProxySettings MaterialSettings;
+
+	/** If Value is > 1, Multisample output baked textures by this amount in each direction (eg 4 == 16x supersampling) */
+	UPROPERTY(BlueprintReadWrite, EditAnywhere, AdvancedDisplay, Category = MaterialSettings, meta = (ClampMin = "0", ClampMax = "8", UIMin = "0", UIMax = "4"))
+	int32 MultiSamplingAA = 0;
+
+	UPROPERTY(BlueprintReadWrite, EditAnywhere, AdvancedDisplay, Category = MaterialSettings, meta = (ClampMin = "5.0", ClampMax = "160.0"))
+	float CaptureFieldOfView = 30.0f;
+
+	UPROPERTY(BlueprintReadWrite, EditAnywhere, AdvancedDisplay, Category = MaterialSettings, meta = (ClampMin = "0.001", ClampMax = "1000.0"))
+	float NearPlaneDist = 1.0f;
+
+	/** If Value is zero, use MaterialSettings resolution, otherwise override the render capture resolution */
+	UPROPERTY(BlueprintReadWrite, EditAnywhere, AdvancedDisplay, Category = MaterialSettings, meta = (ClampMin = "128"))
+	int32 RenderCaptureResolution = 0;
+
+	/** Equality operator. */
+	bool operator==(const FMeshApproximationSettings& Other) const
+	{
+		return ApproximationAccuracy == Other.ApproximationAccuracy
+			&& ClampVoxelDimension == Other.ClampVoxelDimension
+			&& WindingThreshold == Other.WindingThreshold
+			&& bFillGaps == Other.bFillGaps
+			&& GapDistance == Other.GapDistance
+			&& SimplifyMethod == Other.SimplifyMethod
+			&& TargetTriCount == Other.TargetTriCount
+			&& TrianglesPerM == Other.TrianglesPerM
+
+			&& MaterialSettings == Other.MaterialSettings
+			&& MultiSamplingAA == Other.MultiSamplingAA
+			&& CaptureFieldOfView == Other.CaptureFieldOfView
+			&& NearPlaneDist == Other.NearPlaneDist
+			&& RenderCaptureResolution == Other.RenderCaptureResolution;
+	}
+
+	/** Inequality. */
+	bool operator!=(const FMeshApproximationSettings& Other) const
+	{
+		return !(*this == Other);
+	}
+
+#if WITH_EDITORONLY_DATA
+	/** Handles deprecated properties */
+	void PostLoadDeprecated() {}		// none currently
+#endif
 };
\ No newline at end of file
diff --git a/Engine/Source/Runtime/Engine/Private/WorldPartition/HLOD/HLODLayer.cpp b/Engine/Source/Runtime/Engine/Private/WorldPartition/HLOD/HLODLayer.cpp
--- a/Engine/Source/Runtime/Engine/Private/WorldPartition/HLOD/HLODLayer.cpp
+++ b/Engine/Source/Runtime/Engine/Private/WorldPartition/HLOD/HLODLayer.cpp
@@ -57,6 +57,11 @@ uint32 UHLODLayer::GetCRC() const
 		Ar << This.MeshSimplifySettings;
 		UE_LOG(LogHLODLayer, VeryVerbose, TEXT(" - MeshSimplifySettings = %d"), Ar.GetCrc());
 	}
+	else if (LayerType == EHLODLayerType::MeshApproximate)
+	{
+		Ar << This.MeshApproximationSettings;
+		UE_LOG(LogHLODLayer, VeryVerbose, TEXT(" - MeshApproximationSettings = %d"), Ar.GetCrc());
+	}
 
 	Ar << This.CellSize;
 	UE_LOG(LogHLODLayer, VeryVerbose, TEXT(" - CellSize = %d"), Ar.GetCrc());
diff --git a/Engine/Source/Runtime/Engine/Public/WorldPartition/HLOD/HLODLayer.h b/Engine/Source/Runtime/Engine/Public/WorldPartition/HLOD/HLODLayer.h
--- a/Engine/Source/Runtime/Engine/Public/WorldPartition/HLOD/HLODLayer.h
+++ b/Engine/Source/Runtime/Engine/Public/WorldPartition/HLOD/HLODLayer.h
@@ -9,6 +9,7 @@
 
 #if WITH_EDITOR
 #include "Engine/MeshMerging.h"
+#include "Engine/MeshMerging.h"
 #endif
 
 #include "HLODLayer.generated.h"
@@ -19,9 +20,10 @@ class AActor;
 UENUM()
 enum class EHLODLayerType : uint8
 {
-	Instancing = 0		UMETA(DisplayName = "Instancing"),
-	MeshMerge = 1		UMETA(DisplayName = "Merged Mesh"),
-	MeshSimplify = 2	UMETA(DisplayName = "Simplified Mesh")
+	Instancing = 0			UMETA(DisplayName = "Instancing"),
+	MeshMerge = 1			UMETA(DisplayName = "Merged Mesh"),
+	MeshSimplify = 2		UMETA(DisplayName = "Simplified Mesh"),
+	MeshApproximate = 3		UMETA(DisplayName = "Approximated Mesh"),
 };
 
 UCLASS(Blueprintable, Config=Engine, PerObjectConfig)
@@ -41,6 +43,7 @@ public:
 	EHLODLayerType GetLayerType() const { return LayerType; }
 	const FMeshMergingSettings& GetMeshMergeSettings() const { return MeshMergeSettings; }
 	const FMeshProxySettings& GetMeshSimplifySettings() const { return MeshSimplifySettings; }
+	const FMeshApproximationSettings& GetMeshApproximationSettings() const { return MeshApproximationSettings; }
 	const TSoftObjectPtr<UMaterial>& GetHLODMaterial() const { return HLODMaterial; }
 	FName GetRuntimeGrid(uint32 InHLODLevel) const;
 	int32 GetCellSize() const { return bAlwaysLoaded ? 0 : CellSize; }
@@ -65,6 +68,10 @@ private:
 	UPROPERTY(EditAnywhere, Config, Category=HLOD, meta = (EditConditionHides, EditCondition = "LayerType == EHLODLayerType::MeshSimplify"))
 	FMeshProxySettings MeshSimplifySettings;
 
+	/** Mesh approximation settings - Used when this layer is of type EHLODLayerType::MeshApproximation */
+	UPROPERTY(EditAnywhere, Config, Category = HLOD, meta = (EditConditionHides, EditCondition = "LayerType == EHLODLayerType::MeshApproximate"))
+	FMeshApproximationSettings MeshApproximationSettings;
+
 	/** Material that will be used by the generated HLOD static mesh */
 	UPROPERTY(EditAnywhere, Config, AdvancedDisplay, Category=HLOD, meta = (EditConditionHides, EditCondition = "LayerType == EHLODLayerType::MeshMerge || LayerType == EHLODLayerType::MeshSimplify"))
 	TSoftObjectPtr<UMaterial> HLODMaterial;
