Make sure all assets compilation have completed before building an HLOD

diff --git a/Engine/Source/Developer/WorldPartitionHLODUtilities/Private/HLODBuilder.cpp b/Engine/Source/Developer/WorldPartitionHLODUtilities/Private/HLODBuilder.cpp
--- a/Engine/Source/Developer/WorldPartitionHLODUtilities/Private/HLODBuilder.cpp
+++ b/Engine/Source/Developer/WorldPartitionHLODUtilities/Private/HLODBuilder.cpp
@@ -8,6 +8,8 @@
 #include "WorldPartition/HLOD/HLODActor.h"
 #include "WorldPartition/HLOD/HLODLayer.h"
 
+#include "AssetCompilingManager.h"
+
 DEFINE_LOG_CATEGORY(LogHLODBuilder);
 
 
@@ -19,6 +21,9 @@ void FHLODBuilder::Build(AWorldPartitionHLOD* InHLODActor, const UHLODLayer* InH
 		return;
 	}
 
+	// Before capturing the scene, make sure all assets are finished compiling
+	FAssetCompilingManager::Get().FinishAllCompilation();
+
 	TArray<UPrimitiveComponent*> HLODPrimitives = CreateComponents(InHLODActor, InHLODLayer, SubComponents);
 	HLODPrimitives.RemoveSwap(nullptr);
 
