Copy PerInstanceData when creating instanced HLODs

diff --git a/Engine/Source/Developer/WorldPartitionHLODUtilities/Private/HLODBuilderInstancing.cpp b/Engine/Source/Developer/WorldPartitionHLODUtilities/Private/HLODBuilderInstancing.cpp
--- a/Engine/Source/Developer/WorldPartitionHLODUtilities/Private/HLODBuilderInstancing.cpp
+++ b/Engine/Source/Developer/WorldPartitionHLODUtilities/Private/HLODBuilderInstancing.cpp
@@ -62,6 +62,17 @@ private:
 	uint32						Hash;
 };
 
+// Data used to batch instances
+struct FInstancingData
+{
+	int32							NumInstances;
+
+	TArray<FTransform>				InstancesTransforms;
+
+	int32							NumCustomDataFloats;
+	TArray<float>					InstancesCustomData;
+};
+
 
 TArray<UPrimitiveComponent*> FHLODBuilder_Instancing::CreateComponents(AWorldPartitionHLOD* InHLODActor, const UHLODLayer* InHLODLayer, const TArray<UPrimitiveComponent*>& InSubComponents)
 {
@@ -70,47 +81,82 @@ TArray<UPrimitiveComponent*> FHLODBuilder_Instancing::CreateComponents(AWorldPar
 	TArray<UPrimitiveComponent*> Components;
 
 	// Gather all meshes to instantiate along with their transforms
-	TMap<FInstancingKey, TArray<UPrimitiveComponent*>> Instances;
+	TMap<FInstancingKey, FInstancingData> InstancesData;
 	for (UPrimitiveComponent* Primitive : InSubComponents)
 	{
 		if (UStaticMeshComponent* SMC = Cast<UStaticMeshComponent>(Primitive))
 		{
-			Instances.FindOrAdd(FInstancingKey(SMC)).Add(SMC);
+			FInstancingData& InstancingData = InstancesData.FindOrAdd(FInstancingKey(SMC));
+
+			if (UInstancedStaticMeshComponent* ISMC = Cast<UInstancedStaticMeshComponent>(SMC))
+			{
+				InstancingData.NumCustomDataFloats = FMath::Max(InstancingData.NumCustomDataFloats, ISMC->PerInstanceSMCustomData.Num());
+				InstancingData.NumInstances += ISMC->GetInstanceCount();
+			}
+			else
+			{
+				InstancingData.NumInstances++;
+			}
 		}
 	}
 
-	// Create an ISMC for each SM asset we found
-	for (const auto& Entry : Instances)
+	// Resize arrays
+	for (auto& Entry : InstancesData)
 	{
-		const FInstancingKey EntryInstancingKey = Entry.Key;
-		const TArray<UPrimitiveComponent*>& EntryComponents = Entry.Value;
-			
-		UInstancedStaticMeshComponent* Component = NewObject<UInstancedStaticMeshComponent>(InHLODActor);
-		EntryInstancingKey.ApplyTo(Component);
-		Component->SetForcedLodModel(Component->GetStaticMesh()->GetNumLODs());
-
-		DisableCollisions(Component);
+		const FInstancingKey& EntryInstancingKey = Entry.Key;
+		FInstancingData& EntryInstancingData = Entry.Value;
 
-		// Add all instances
-		TArray<FTransform> InstanceTransforms;
-		for (UPrimitiveComponent* SMC : EntryComponents)
+		EntryInstancingData.InstancesTransforms.Reset(EntryInstancingData.NumInstances);
+		EntryInstancingData.InstancesCustomData.Reset(EntryInstancingData.NumInstances * EntryInstancingData.NumCustomDataFloats);
+	}
+	
+	// Append all transforms & per instance custom data
+	for (UPrimitiveComponent* Primitive : InSubComponents)
+	{
+		if (UStaticMeshComponent* SMC = Cast<UStaticMeshComponent>(Primitive))
 		{
-			// If we have an ISMC, retrieve all instances
-			if (UInstancedStaticMeshComponent* InstancedStaticMeshComponent = Cast<UInstancedStaticMeshComponent>(SMC))
+			FInstancingData& InstancingData = InstancesData.FindChecked(FInstancingKey(SMC));
+
+			int32 NumCustomDataFloatsAdded = 0;
+
+			if (UInstancedStaticMeshComponent* ISMC = Cast<UInstancedStaticMeshComponent>(SMC))
 			{
-				for (int32 InstanceIdx = 0; InstanceIdx < InstancedStaticMeshComponent->GetInstanceCount(); InstanceIdx++)
+				// Add transforms
+				for (int32 InstanceIdx = 0; InstanceIdx < ISMC->GetInstanceCount(); InstanceIdx++)
 				{
-					FTransform InstanceTransform;
-					InstancedStaticMeshComponent->GetInstanceTransform(InstanceIdx, InstanceTransform, true);
-					InstanceTransforms.Add(InstanceTransform);
+					FTransform& InstanceTransform = InstancingData.InstancesTransforms.AddDefaulted_GetRef();
+					ISMC->GetInstanceTransform(InstanceIdx, InstanceTransform, true);
 				}
+
+				// Add per instance custom data
+				InstancingData.InstancesCustomData.Append(ISMC->PerInstanceSMCustomData);
+				NumCustomDataFloatsAdded = ISMC->PerInstanceSMCustomData.Num();
 			}
 			else
 			{
-				InstanceTransforms.Add(SMC->GetComponentTransform());
+				InstancingData.InstancesTransforms.Add(SMC->GetComponentTransform());
 			}
+
+			// Add missing custom data, if any
+			InstancingData.InstancesCustomData.AddDefaulted(InstancingData.NumCustomDataFloats - NumCustomDataFloatsAdded);
 		}
-		Component->AddInstances(InstanceTransforms, /*bShouldReturnIndices*/false, /*bWorldSpace*/true);
+	}
+
+	// Create an ISMC for each SM asset we found
+	for (auto& Entry : InstancesData)
+	{
+		const FInstancingKey& EntryInstancingKey = Entry.Key;
+		FInstancingData& EntryInstancingData = Entry.Value;
+
+		UInstancedStaticMeshComponent* Component = NewObject<UInstancedStaticMeshComponent>(InHLODActor);
+		EntryInstancingKey.ApplyTo(Component);
+		Component->SetForcedLodModel(Component->GetStaticMesh()->GetNumLODs());
+
+		DisableCollisions(Component);
+
+		Component->NumCustomDataFloats = EntryInstancingData.NumCustomDataFloats;
+		Component->AddInstances(EntryInstancingData.InstancesTransforms, /*bShouldReturnIndices*/false, /*bWorldSpace*/true);
+		Component->PerInstanceSMCustomData = MoveTemp(EntryInstancingData.InstancesCustomData);
 
 		Components.Add(Component);
 	};
