World Partition HLOD - Add support for nanite meshes & generating virtual textures to the Actors Approximation merge method
* Use the global engine DefaultHLODFlattenMaterial, which expects virtual textures
* Propagate newly added mesh merge options to the IGeometryProcessing_ApproximateActors::FOptions struct
* Added "EmissiveHDRTexture" texture nodes and "UseEmissiveHDR" options to both FlattenMaterial & FlattenMaterial_VT - will be used by the Actor Approx. HLODs

diff --git a/Engine/Source/Developer/WorldPartitionHLODUtilities/Private/HLODBuilderMeshApproximate.cpp b/Engine/Source/Developer/WorldPartitionHLODUtilities/Private/HLODBuilderMeshApproximate.cpp
--- a/Engine/Source/Developer/WorldPartitionHLODUtilities/Private/HLODBuilderMeshApproximate.cpp
+++ b/Engine/Source/Developer/WorldPartitionHLODUtilities/Private/HLODBuilderMeshApproximate.cpp
@@ -9,6 +9,7 @@
 
 #include "Engine/StaticMesh.h"
 #include "Materials/Material.h"
+#include "Materials/MaterialInstance.h"
 
 #include "Modules/ModuleManager.h"
 #include "IGeometryProcessingInterfacesModule.h"
@@ -17,7 +18,7 @@
 
 TArray<UPrimitiveComponent*> FHLODBuilder_MeshApproximate::CreateComponents(AWorldPartitionHLOD* InHLODActor, const UHLODLayer* InHLODLayer, const TArray<UPrimitiveComponent*>& InSubComponents)
 {
-	TRACE_CPUPROFILER_EVENT_SCOPE(FHLODBuilder_MeshSimplify::CreateComponents);
+	TRACE_CPUPROFILER_EVENT_SCOPE(FHLODBuilder_MeshApproximate::CreateComponents);
 
 	TSet<AActor*> Actors;
 	Algo::Transform(InSubComponents, Actors, [](UPrimitiveComponent* InPrimitiveComponent) { return InPrimitiveComponent->GetOwner(); });
@@ -33,17 +34,42 @@ TArray<UPrimitiveComponent*> FHLODBuilder_MeshApproximate::CreateComponents(AWor
 
 	IGeometryProcessing_ApproximateActors::FOptions Options;
 	Options.BasePackagePath = InHLODActor->GetPackage()->GetName();
+	Options.BasePolicy = (UseSettings.OutputType == EMeshApproximationType::MeshShapeOnly) ?
+		IGeometryProcessing_ApproximateActors::EApproximationPolicy::CollisionMesh :
+		IGeometryProcessing_ApproximateActors::EApproximationPolicy::MeshAndGeneratedMaterial;
 	Options.WorldSpaceApproximationAccuracyMeters = UseSettings.ApproximationAccuracy;
+
+	Options.bAutoThickenThinParts = UseSettings.bAttemptAutoThickening;
+	Options.AutoThickenThicknessMeters = UseSettings.TargetMinThicknessMultiplier * UseSettings.ApproximationAccuracy;
+
+	Options.BaseCappingPolicy = IGeometryProcessing_ApproximateActors::EBaseCappingPolicy::NoBaseCapping;
+	if (UseSettings.BaseCapping == EMeshApproximationBaseCappingType::ConvexPolygon)
+	{
+		Options.BaseCappingPolicy = IGeometryProcessing_ApproximateActors::EBaseCappingPolicy::ConvexPolygon;
+	}
+	else if (UseSettings.BaseCapping == EMeshApproximationBaseCappingType::ConvexSolid)
+	{
+		Options.BaseCappingPolicy = IGeometryProcessing_ApproximateActors::EBaseCappingPolicy::ConvexSolid;
+	}
+
 	Options.ClampVoxelDimension = UseSettings.ClampVoxelDimension;
 	Options.WindingThreshold = UseSettings.WindingThreshold;
 	Options.bApplyMorphology = UseSettings.bFillGaps;
 	Options.MorphologyDistanceMeters = UseSettings.GapDistance;
+
+	Options.OcclusionPolicy = (UseSettings.OcclusionMethod == EOccludedGeometryFilteringPolicy::VisibilityBasedFiltering) ?
+		IGeometryProcessing_ApproximateActors::EOcclusionPolicy::VisibilityBased : IGeometryProcessing_ApproximateActors::EOcclusionPolicy::None;
 	Options.FixedTriangleCount = UseSettings.TargetTriCount;
 	if (UseSettings.SimplifyMethod == EMeshApproximationSimplificationPolicy::TrianglesPerArea)
 	{
 		Options.MeshSimplificationPolicy = IGeometryProcessing_ApproximateActors::ESimplificationPolicy::TrianglesPerUnitSqMeter;
 		Options.SimplificationTargetMetric = UseSettings.TrianglesPerM;
 	}
+	else if (UseSettings.SimplifyMethod == EMeshApproximationSimplificationPolicy::GeometricTolerance)
+	{
+		Options.MeshSimplificationPolicy = IGeometryProcessing_ApproximateActors::ESimplificationPolicy::GeometricTolerance;
+		Options.SimplificationTargetMetric = UseSettings.GeometricDeviation;
+	}
 	else
 	{
 		Options.MeshSimplificationPolicy = IGeometryProcessing_ApproximateActors::ESimplificationPolicy::FixedTriangleCount;
@@ -57,23 +83,41 @@ TArray<UPrimitiveComponent*> FHLODBuilder_MeshApproximate::CreateComponents(AWor
 	Options.FieldOfViewDegrees = UseSettings.CaptureFieldOfView;
 	Options.NearPlaneDist = UseSettings.NearPlaneDist;
 
+	Options.bVerbose = UseSettings.bPrintDebugMessages;
+	Options.bWriteDebugMesh = UseSettings.bEmitFullDebugMesh;
+
+	// Nanite settings
+	Options.bGenerateNaniteEnabledMesh = UseSettings.bGenerateNaniteEnabledMesh;
+	Options.NaniteProxyTrianglePercent = UseSettings.NaniteProxyTrianglePercent;
+
+	// Material baking settings
+	Options.BakeMaterial = GEngine->DefaultHLODFlattenMaterial;
+	Options.BaseColorTexParamName = FName("BaseColorTexture");
+	Options.NormalTexParamName = FName("NormalTexture");
+	Options.MetallicTexParamName = FName("MetallicTexture");
+	Options.RoughnessTexParamName = FName("RoughnessTexture");
+	Options.SpecularTexParamName = FName("SpecularTexture");
+	Options.EmissiveTexParamName = FName("EmissiveHDRTexture");
+	Options.bUsePackedMRS = true;
+	Options.PackedMRSTexParamName = FName("PackedTexture");
+
 	// run actor approximation computation
 	IGeometryProcessing_ApproximateActors::FResults Results;
 	ApproxActorsAPI->ApproximateActors(Actors.Array(), Options, Results);
 
-	auto FixupAsset = [&InHLODActor](UObject* Asset)
-	{
-		Asset->ClearFlags(RF_Public | RF_Standalone);
-		Asset->Rename(nullptr, InHLODActor->GetPackage(), REN_NonTransactional | REN_DontCreateRedirectors | REN_ForceNoResetLoaders);
-	};
-	
-	Algo::ForEach(Results.NewMeshAssets, FixupAsset);
-	Algo::ForEach(Results.NewMaterials, FixupAsset);
-	Algo::ForEach(Results.NewTextures, FixupAsset);
-
 	TArray<UPrimitiveComponent*> Components;
 	if (Results.ResultCode == IGeometryProcessing_ApproximateActors::EResultCode::Success)
 	{
+		auto FixupAsset = [&InHLODActor](UObject* Asset)
+		{
+			Asset->ClearFlags(RF_Public | RF_Standalone);
+			Asset->Rename(nullptr, InHLODActor->GetPackage(), REN_NonTransactional | REN_DontCreateRedirectors | REN_ForceNoResetLoaders);
+		};
+	
+		Algo::ForEach(Results.NewMeshAssets, FixupAsset);
+		Algo::ForEach(Results.NewMaterials, FixupAsset);
+		Algo::ForEach(Results.NewTextures, FixupAsset);
+
 		for (UStaticMesh* StaticMesh : Results.NewMeshAssets)
 		{
 			UStaticMeshComponent* Component = NewObject<UStaticMeshComponent>(InHLODActor);
@@ -81,7 +125,42 @@ TArray<UPrimitiveComponent*> FHLODBuilder_MeshApproximate::CreateComponents(AWor
 			DisableCollisions(Component);
 
 			Components.Add(Component);
-		}		
+		}
+
+		for (UMaterialInterface* Material : Results.NewMaterials)
+		{
+			UMaterialInstance* MaterialInst = CastChecked<UMaterialInstance>(Material);
+
+			FStaticParameterSet StaticParameterSet;
+			
+			auto SetStaticSwitch = [&StaticParameterSet](FName ParamName, bool bSet)
+			{
+				if (bSet)
+				{
+					FStaticSwitchParameter SwitchParameter;
+					SwitchParameter.ParameterInfo.Name = ParamName;
+					SwitchParameter.Value = true;
+					SwitchParameter.bOverride = true;
+					StaticParameterSet.StaticSwitchParameters.Add(SwitchParameter);
+				}
+			};
+
+			// Set proper switches needed by our base flatten material
+			SetStaticSwitch("UseBaseColor", Options.bBakeBaseColor);
+			SetStaticSwitch("UseRoughness", Options.bBakeRoughness);
+			SetStaticSwitch("UseMetallic", Options.bBakeMetallic);
+			SetStaticSwitch("UseSpecular", Options.bBakeSpecular);
+			SetStaticSwitch("UseEmissiveHDR", Options.bBakeEmissive);
+			SetStaticSwitch("UseNormal", Options.bBakeNormalMap);
+			SetStaticSwitch("PackMetallic", Options.bUsePackedMRS);
+			SetStaticSwitch("PackSpecular", Options.bUsePackedMRS);
+			SetStaticSwitch("PackRoughness", Options.bUsePackedMRS);
+
+			// Force initializing the static permutations according to the switches we have set
+			MaterialInst->UpdateStaticPermutation(StaticParameterSet);
+			MaterialInst->InitStaticPermutation();
+			MaterialInst->PostEditChange();
+		}
 	}
 
 	return Components;
