World Partition HLODs - Rely on FISMComponentDescriptor to perform batching for Instanced HLODs

diff --git a/Engine/Source/Developer/WorldPartitionHLODUtilities/Private/HLODBuilderInstancing.cpp b/Engine/Source/Developer/WorldPartitionHLODUtilities/Private/HLODBuilderInstancing.cpp
--- a/Engine/Source/Developer/WorldPartitionHLODUtilities/Private/HLODBuilderInstancing.cpp
+++ b/Engine/Source/Developer/WorldPartitionHLODUtilities/Private/HLODBuilderInstancing.cpp
@@ -8,71 +8,39 @@
 #include "WorldPartition/HLOD/HLODActor.h"
 #include "WorldPartition/HLOD/HLODLayer.h"
 
+#include "ISMPartition/ISMComponentDescriptor.h"
 
-// We want to merge all SMC that are using the same static mesh
-// However, we must also take material overiddes into account.
-struct FInstancingKey
+namespace 
 {
-	FInstancingKey(const UStaticMeshComponent* SMC)
+	// Instance batcher class based on FISMComponentDescriptor
+	struct FCustomISMComponentDescriptor : public FISMComponentDescriptor
 	{
-		FHashBuilder HashBuilder;
-
-		StaticMesh = SMC->GetStaticMesh();
-		HashBuilder << StaticMesh;
-
-		const int32 NumMaterials = SMC->GetNumMaterials();
-		Materials.Reserve(NumMaterials);
-
-		for (int32 MaterialIndex = 0; MaterialIndex < NumMaterials; ++MaterialIndex)
+		FCustomISMComponentDescriptor(UStaticMeshComponent* SMC)
 		{
-			UMaterialInterface* Material = SMC->GetMaterial(MaterialIndex);
+			InitFrom(SMC, false);
 
-			Materials.Add(Material);
-			HashBuilder << Material;
-		}
-
-		Hash = HashBuilder.GetHash();
-	}
-
-	void ApplyTo(UStaticMeshComponent* SMC) const
-	{
-		// Set static mesh
-		SMC->SetStaticMesh(StaticMesh);
+			// We'll always want to spawn ISMC, even if our source components are all SMC
+			ComponentClass = UInstancedStaticMeshComponent::StaticClass();
 			
-		// Set material overrides
-		for (int32 MaterialIndex = 0; MaterialIndex < Materials.Num(); ++MaterialIndex)
-		{
-			SMC->SetMaterial(MaterialIndex, Materials[MaterialIndex]);
-		}
-	}
+			// For now, ignore ray tracing group ID when batching
+			// We may want to expose an instance batching option to control this
+			RayTracingGroupId = FPrimitiveSceneProxy::InvalidRayTracingGroupId;
 
-	friend uint32 GetTypeHash(const FInstancingKey& Key)
-	{
-		return Key.Hash;
-	}
+			ComputeHash();
+		}
+	};
 
-	bool operator==(const FInstancingKey& Other) const
+	// Store batched instances data
+	struct FInstancingData
 	{
-		return Hash == Other.Hash && StaticMesh == Other.StaticMesh && Materials == Other.Materials;
-	}
+		int32							NumInstances;
 
-private:
-	UStaticMesh*				StaticMesh;
-	TArray<UMaterialInterface*>	Materials;
-	uint32						Hash;
-};
-
-// Data used to batch instances
-struct FInstancingData
-{
-	int32							NumInstances;
-
-	TArray<FTransform>				InstancesTransforms;
-
-	int32							NumCustomDataFloats;
-	TArray<float>					InstancesCustomData;
-};
+		TArray<FTransform>				InstancesTransforms;
 
+		int32							NumCustomDataFloats;
+		TArray<float>					InstancesCustomData;
+	};
+}
 
 TArray<UPrimitiveComponent*> FHLODBuilder_Instancing::CreateComponents(AWorldPartitionHLOD* InHLODActor, const UHLODLayer* InHLODLayer, const TArray<UPrimitiveComponent*>& InSubComponents)
 {
@@ -80,13 +48,14 @@ TArray<UPrimitiveComponent*> FHLODBuilder_Instancing::CreateComponents(AWorldPar
 
 	TArray<UPrimitiveComponent*> Components;
 
-	// Gather all meshes to instantiate along with their transforms
-	TMap<FInstancingKey, FInstancingData> InstancesData;
+	// Prepare instance batches
+	TMap<FISMComponentDescriptor, FInstancingData> InstancesData;
 	for (UPrimitiveComponent* Primitive : InSubComponents)
 	{
 		if (UStaticMeshComponent* SMC = Cast<UStaticMeshComponent>(Primitive))
 		{
-			FInstancingData& InstancingData = InstancesData.FindOrAdd(FInstancingKey(SMC));
+			FCustomISMComponentDescriptor ISMComponentDescriptor(SMC);
+			FInstancingData& InstancingData = InstancesData.FindOrAdd(ISMComponentDescriptor);
 
 			if (UInstancedStaticMeshComponent* ISMC = Cast<UInstancedStaticMeshComponent>(SMC))
 			{
@@ -103,7 +72,7 @@ TArray<UPrimitiveComponent*> FHLODBuilder_Instancing::CreateComponents(AWorldPar
 	// Resize arrays
 	for (auto& Entry : InstancesData)
 	{
-		const FInstancingKey& EntryInstancingKey = Entry.Key;
+		const FISMComponentDescriptor& ISMComponentDescriptor = Entry.Key;
 		FInstancingData& EntryInstancingData = Entry.Value;
 
 		EntryInstancingData.InstancesTransforms.Reset(EntryInstancingData.NumInstances);
@@ -115,7 +84,8 @@ TArray<UPrimitiveComponent*> FHLODBuilder_Instancing::CreateComponents(AWorldPar
 	{
 		if (UStaticMeshComponent* SMC = Cast<UStaticMeshComponent>(Primitive))
 		{
-			FInstancingData& InstancingData = InstancesData.FindChecked(FInstancingKey(SMC));
+			FCustomISMComponentDescriptor ISMComponentDescriptor(SMC);
+			FInstancingData& InstancingData = InstancesData.FindChecked(ISMComponentDescriptor);
 
 			int32 NumCustomDataFloatsAdded = 0;
 
@@ -145,11 +115,10 @@ TArray<UPrimitiveComponent*> FHLODBuilder_Instancing::CreateComponents(AWorldPar
 	// Create an ISMC for each SM asset we found
 	for (auto& Entry : InstancesData)
 	{
-		const FInstancingKey& EntryInstancingKey = Entry.Key;
+		const FISMComponentDescriptor& ISMComponentDescriptor = Entry.Key;
 		FInstancingData& EntryInstancingData = Entry.Value;
 
-		UInstancedStaticMeshComponent* Component = NewObject<UInstancedStaticMeshComponent>(InHLODActor);
-		EntryInstancingKey.ApplyTo(Component);
+		UInstancedStaticMeshComponent* Component = ISMComponentDescriptor.CreateComponent(InHLODActor);
 		Component->SetForcedLodModel(Component->GetStaticMesh()->GetNumLODs());
 
 		DisableCollisions(Component);
