World Partition HLODs - Promotion changes
* An HLOD actor can use the "Location" grid placement if all it's source are using this placement option.
* Properly accessing grid placement of source actors using the actor desc view

diff --git a/Engine/Source/Developer/WorldPartitionHLODUtilities/Private/WorldPartitionHLODUtilities.cpp b/Engine/Source/Developer/WorldPartitionHLODUtilities/Private/WorldPartitionHLODUtilities.cpp
--- a/Engine/Source/Developer/WorldPartitionHLODUtilities/Private/WorldPartitionHLODUtilities.cpp
+++ b/Engine/Source/Developer/WorldPartitionHLODUtilities/Private/WorldPartitionHLODUtilities.cpp
@@ -145,26 +145,38 @@ static uint32 ComputeHLODHash(AWorldPartitionHLOD* InHLODActor, const TArray<FWo
 
 TArray<AWorldPartitionHLOD*> FWorldPartitionHLODUtilities::CreateHLODActors(FHLODCreationContext& InCreationContext, const FHLODCreationParams& InCreationParams, const TSet<FActorInstance>& InActors, const TArray<const UDataLayer*>& InDataLayers)
 {
-	TMap<UHLODLayer*, TArray<FGuid>> HLODLayersActors;
+	struct FSubActorsInfo
+	{
+		TArray<FGuid>			SubActors;
+		EActorGridPlacement		GridPlacement = EActorGridPlacement::Location;
+	};
+	TMap<UHLODLayer*, FSubActorsInfo> SubActorsInfos;
+
 	for (const FActorInstance& ActorInstance : InActors)
 	{
-		FWorldPartitionActorDesc& ActorDesc = InCreationParams.WorldPartition->GetActorDescChecked(ActorInstance.Actor);
-		if (ActorDesc.GetActorIsHLODRelevant())
+		const FWorldPartitionActorDescView& ActorDescView = ActorInstance.GetActorDescView();
+		if (ActorDescView.GetActorIsHLODRelevant())
 		{
-			UHLODLayer* HLODLayer = UHLODLayer::GetHLODLayer(ActorDesc, InCreationParams.WorldPartition);
+			UHLODLayer* HLODLayer = UHLODLayer::GetHLODLayer(ActorDescView, InCreationParams.WorldPartition);
 			if (HLODLayer)
 			{
-				HLODLayersActors.FindOrAdd(HLODLayer).Add(ActorInstance.Actor);
+				FSubActorsInfo& SubActorsInfo = SubActorsInfos.FindOrAdd(HLODLayer);
+
+				SubActorsInfo.SubActors.Add(ActorInstance.Actor);
+				if (ActorDescView.GetGridPlacement() == EActorGridPlacement::Bounds)
+				{
+					SubActorsInfo.GridPlacement = EActorGridPlacement::Bounds;
+				}
 			}
 		}
 	}
 
 	TArray<AWorldPartitionHLOD*> HLODActors;
-	for (const auto& Pair : HLODLayersActors)
+	for (const auto& Pair : SubActorsInfos)
 	{
 		const UHLODLayer* HLODLayer = Pair.Key;
-		const TArray<FGuid> HLODLayerActors = Pair.Value;
-		check(!HLODLayerActors.IsEmpty());
+		const FSubActorsInfo& SubActorsInfo = Pair.Value;
+		check(!SubActorsInfo.SubActors.IsEmpty());
 
 		// Compute HLODActor hash
 		uint64 CellHash = FHLODActorDesc::ComputeCellHash(HLODLayer->GetName(), InCreationParams.GridIndexX, InCreationParams.GridIndexY, InCreationParams.GridIndexZ, InCreationParams.DataLayersID);
@@ -218,11 +230,11 @@ TArray<AWorldPartitionHLOD*> FWorldPartitionHLODUtilities::CreateHLODActors(FHLO
 
 		// Sub actors
 		{
-			bool bSubActorsChanged = HLODActor->GetSubActors().Num() != HLODLayerActors.Num();
+			bool bSubActorsChanged = HLODActor->GetSubActors().Num() != SubActorsInfo.SubActors.Num();
 			if (!bSubActorsChanged)
 			{
 				TArray<FGuid> A = HLODActor->GetSubActors();
-				TArray<FGuid> B = HLODLayerActors;
+				TArray<FGuid> B = SubActorsInfo.SubActors;
 				A.Sort();
 				B.Sort();
 				bSubActorsChanged = A != B;
@@ -230,7 +242,7 @@ TArray<AWorldPartitionHLOD*> FWorldPartitionHLODUtilities::CreateHLODActors(FHLO
 
 			if (bSubActorsChanged)
 			{
-				HLODActor->SetSubActors(HLODLayerActors);
+				HLODActor->SetSubActors(SubActorsInfo.SubActors);
 				bIsDirty = true;
 			}
 		}
@@ -243,6 +255,13 @@ TArray<AWorldPartitionHLOD*> FWorldPartitionHLODUtilities::CreateHLODActors(FHLO
 			bIsDirty = true;
 		}
 
+		// Runtime grid placement
+		if (HLODActor->GetGridPlacement() != SubActorsInfo.GridPlacement)
+		{
+			HLODActor->SetGridPlacement(SubActorsInfo.GridPlacement);
+			bIsDirty = true;
+		}
+
 		// HLOD level
 		if (HLODActor->GetLODLevel() != InCreationParams.HLODLevel)
 		{
diff --git a/Engine/Source/Runtime/Engine/Private/WorldPartition/HLOD/HLODActor.cpp b/Engine/Source/Runtime/Engine/Private/WorldPartition/HLOD/HLODActor.cpp
--- a/Engine/Source/Runtime/Engine/Private/WorldPartition/HLOD/HLODActor.cpp
+++ b/Engine/Source/Runtime/Engine/Private/WorldPartition/HLOD/HLODActor.cpp
@@ -78,7 +78,7 @@ bool AWorldPartitionHLOD::IsHiddenEd() const
 
 EActorGridPlacement AWorldPartitionHLOD::GetGridPlacement() const
 {
-	return SubActorsHLODLayer && SubActorsHLODLayer->IsAlwaysLoaded() ? EActorGridPlacement::AlwaysLoaded : EActorGridPlacement::Bounds;
+	return SubActorsHLODLayer && SubActorsHLODLayer->IsAlwaysLoaded() ? EActorGridPlacement::AlwaysLoaded : Super::GetGridPlacement();
 }
 
 EActorGridPlacement AWorldPartitionHLOD::GetDefaultGridPlacement() const
diff --git a/Engine/Source/Runtime/Engine/Private/WorldPartition/HLOD/HLODLayer.cpp b/Engine/Source/Runtime/Engine/Private/WorldPartition/HLOD/HLODLayer.cpp
--- a/Engine/Source/Runtime/Engine/Private/WorldPartition/HLOD/HLODLayer.cpp
+++ b/Engine/Source/Runtime/Engine/Private/WorldPartition/HLOD/HLODLayer.cpp
@@ -16,6 +16,7 @@
 
 #include "WorldPartition/WorldPartition.h"
 #include "WorldPartition/WorldPartitionActorDesc.h"
+#include "WorldPartition/WorldPartitionActorDescView.h"
 #endif
 
 DEFINE_LOG_CATEGORY_STATIC(LogHLODLayer, Log, All);
@@ -103,7 +104,7 @@ UHLODLayer* UHLODLayer::GetHLODLayer(const AActor* InActor)
 	return nullptr;
 }
 
-UHLODLayer* UHLODLayer::GetHLODLayer(const FWorldPartitionActorDesc& InActorDesc, const UWorldPartition* InWorldPartition)
+UHLODLayer* UHLODLayer::GetHLODLayer(const FWorldPartitionActorDescView& InActorDesc, const UWorldPartition* InWorldPartition)
 {
 	check(InWorldPartition);
 
@@ -123,6 +124,11 @@ UHLODLayer* UHLODLayer::GetHLODLayer(const FWorldPartitionActorDesc& InActorDesc
 	return nullptr;
 }
 
+UHLODLayer* UHLODLayer::GetHLODLayer(const FWorldPartitionActorDesc& InActorDesc, const UWorldPartition* InWorldPartition)
+{
+	return GetHLODLayer(FWorldPartitionActorDescView(&InActorDesc), InWorldPartition);
+}
+
 #endif // WITH_EDITOR
 
 #if WITH_EDITORONLY_DATA
diff --git a/Engine/Source/Runtime/Engine/Public/WorldPartition/HLOD/HLODLayer.h b/Engine/Source/Runtime/Engine/Public/WorldPartition/HLOD/HLODLayer.h
--- a/Engine/Source/Runtime/Engine/Public/WorldPartition/HLOD/HLODLayer.h
+++ b/Engine/Source/Runtime/Engine/Public/WorldPartition/HLOD/HLODLayer.h
@@ -12,15 +12,15 @@
 
 #if WITH_EDITOR
 #include "Engine/MeshMerging.h"
-#include "Engine/MeshMerging.h"
 #endif
 
 #include "HLODLayer.generated.h"
 
+class AActor;
 class UMaterial;
 class UWorldPartition;
-class AActor;
 class FWorldPartitionActorDesc;
+class FWorldPartitionActorDescView;
 
 UENUM()
 enum class EHLODLayerType : uint8
@@ -40,6 +40,7 @@ class ENGINE_API UHLODLayer : public UObject
 public:
 	static UHLODLayer* GetHLODLayer(const AActor* InActor);
 	static UHLODLayer* GetHLODLayer(const FWorldPartitionActorDesc& InActorDesc, const UWorldPartition* InWorldPartition);
+	static UHLODLayer* GetHLODLayer(const FWorldPartitionActorDescView& InActorDescView, const UWorldPartition* InWorldPartition);
 
 	uint32 GetCRC() const;
 #endif
diff --git a/Engine/Source/Runtime/Engine/Public/WorldPartition/WorldPartitionActorCluster.h b/Engine/Source/Runtime/Engine/Public/WorldPartition/WorldPartitionActorCluster.h
--- a/Engine/Source/Runtime/Engine/Public/WorldPartition/WorldPartitionActorCluster.h
+++ b/Engine/Source/Runtime/Engine/Public/WorldPartition/WorldPartitionActorCluster.h
@@ -68,7 +68,7 @@ struct FActorClusterInstance
 /**
  * Represents one actor and his owning container instance 
  */
-struct FActorInstance
+struct ENGINE_API FActorInstance
 {
 	FGuid Actor;
 	const FActorContainerInstance* ContainerInstance;
