World Partition HLOD Builder - Only wait for async compilation if necessary

diff --git a/Engine/Source/Developer/WorldPartitionHLODUtilities/Private/HLODBuilder.h b/Engine/Source/Developer/WorldPartitionHLODUtilities/Private/HLODBuilder.h
--- a/Engine/Source/Developer/WorldPartitionHLODUtilities/Private/HLODBuilder.h
+++ b/Engine/Source/Developer/WorldPartitionHLODUtilities/Private/HLODBuilder.h
@@ -19,6 +19,8 @@ class FHLODBuilder
 {
 public:
 	virtual ~FHLODBuilder() {}
+	
+	virtual bool RequiresCompiledAssets() const { return true; }
 
 	virtual TArray<UPrimitiveComponent*> CreateComponents(AWorldPartitionHLOD* InHLODActor, const UHLODLayer* InHLODLayer, const TArray<UPrimitiveComponent*>& InSubComponents) = 0;
 
diff --git a/Engine/Source/Developer/WorldPartitionHLODUtilities/Private/HLODBuilderInstancing.h b/Engine/Source/Developer/WorldPartitionHLODUtilities/Private/HLODBuilderInstancing.h
--- a/Engine/Source/Developer/WorldPartitionHLODUtilities/Private/HLODBuilderInstancing.h
+++ b/Engine/Source/Developer/WorldPartitionHLODUtilities/Private/HLODBuilderInstancing.h
@@ -9,5 +9,7 @@
  */
 class FHLODBuilder_Instancing : public FHLODBuilder
 {
+	virtual bool RequiresCompiledAssets() const override { return false; }
+
 	virtual TArray<UPrimitiveComponent*> CreateComponents(AWorldPartitionHLOD* InHLODActor, const UHLODLayer* InHLODLayer, const TArray<UPrimitiveComponent*>& InSubComponents) override;
 };
diff --git a/Engine/Source/Developer/WorldPartitionHLODUtilities/Private/WorldPartitionHLODUtilities.cpp b/Engine/Source/Developer/WorldPartitionHLODUtilities/Private/WorldPartitionHLODUtilities.cpp
--- a/Engine/Source/Developer/WorldPartitionHLODUtilities/Private/WorldPartitionHLODUtilities.cpp
+++ b/Engine/Source/Developer/WorldPartitionHLODUtilities/Private/WorldPartitionHLODUtilities.cpp
@@ -18,7 +18,7 @@
 #include "Engine/StaticMesh.h"
 #include "Engine/HLODProxy.h"
 #include "Materials/Material.h"
-#include "StaticMeshCompiler.h"
+#include "AssetCompilingManager.h"
 
 #include "LevelInstance/LevelInstanceActor.h"
 #include "LevelInstance/LevelInstanceSubsystem.h"
@@ -66,9 +66,6 @@ static bool LoadSubActors(AWorldPartitionHLOD* InHLODActor, TArray<FWorldPartiti
 		}
 	}
 
-	// Wait for compilation to finish
-	FStaticMeshCompilingManager::Get().FinishAllCompilation();
-
 	return !bIsDirty;
 }
 
@@ -342,6 +339,12 @@ uint32 FWorldPartitionHLODUtilities::BuildHLOD(AWorldPartitionHLOD* InHLODActor)
 
 	if (ensure(HLODBuilder))
 	{
+		if (HLODBuilder->RequiresCompiledAssets())
+		{
+			// Wait for compilation to finish
+			FAssetCompilingManager::Get().FinishAllCompilation();
+		}
+
 		HLODBuilder->Build(InHLODActor, HLODLayer, SubActors);
 	}
 			
