Make sure instancing HLOD layer groups components by their RayTracingGroupId field to make sure that information is still available to Lumen.

diff --git a/Engine/Source/Developer/WorldPartitionHLODUtilities/Private/HLODBuilderInstancing.cpp b/Engine/Source/Developer/WorldPartitionHLODUtilities/Private/HLODBuilderInstancing.cpp
--- a/Engine/Source/Developer/WorldPartitionHLODUtilities/Private/HLODBuilderInstancing.cpp
+++ b/Engine/Source/Developer/WorldPartitionHLODUtilities/Private/HLODBuilderInstancing.cpp
@@ -21,10 +21,6 @@ namespace
 
 			// We'll always want to spawn ISMC, even if our source components are all SMC
 			ComponentClass = UInstancedStaticMeshComponent::StaticClass();
-			
-			// For now, ignore ray tracing group ID when batching
-			// We may want to expose an instance batching option to control this
-			RayTracingGroupId = FPrimitiveSceneProxy::InvalidRayTracingGroupId;
 
 			ComputeHash();
 		}
diff --git a/Engine/Source/Developer/WorldPartitionHLODUtilities/Private/WorldPartitionHLODUtilities.cpp b/Engine/Source/Developer/WorldPartitionHLODUtilities/Private/WorldPartitionHLODUtilities.cpp
--- a/Engine/Source/Developer/WorldPartitionHLODUtilities/Private/WorldPartitionHLODUtilities.cpp
+++ b/Engine/Source/Developer/WorldPartitionHLODUtilities/Private/WorldPartitionHLODUtilities.cpp
@@ -74,7 +74,7 @@ static uint32 ComputeHLODHash(AWorldPartitionHLOD* InHLODActor, const TArray<FWo
 	FArchiveCrc32 Ar;
 
 	// Base key, changing this will force a rebuild of all HLODs
-	FString HLODBaseKey = "C2690B6E883841EE851A33165A3B147D";
+	FString HLODBaseKey = "E7DBB5E1B39947FB84E72A38074B7711";
 	Ar << HLODBaseKey;
 
 	// HLOD Layer
