Support batching components using HLODBatchingPolicy == EHLODBatchingPolicy::Instancing as instanced in the MeshApproximate & MeshSimplify HLOD builders

diff --git a/Engine/Source/Developer/WorldPartitionHLODUtilities/Private/HLODBuilderInstancing.h b/Engine/Source/Developer/WorldPartitionHLODUtilities/Private/HLODBuilderInstancing.h
--- a/Engine/Source/Developer/WorldPartitionHLODUtilities/Private/HLODBuilderInstancing.h
+++ b/Engine/Source/Developer/WorldPartitionHLODUtilities/Private/HLODBuilderInstancing.h
@@ -9,6 +9,7 @@
  */
 class FHLODBuilder_Instancing : public FHLODBuilder
 {
+public:
 	virtual bool RequiresCompiledAssets() const override { return false; }
 
 	virtual TArray<UPrimitiveComponent*> CreateComponents(AWorldPartitionHLOD* InHLODActor, const UHLODLayer* InHLODLayer, const TArray<UPrimitiveComponent*>& InSubComponents) override;
diff --git a/Engine/Source/Developer/WorldPartitionHLODUtilities/Private/HLODBuilderMeshApproximate.cpp b/Engine/Source/Developer/WorldPartitionHLODUtilities/Private/HLODBuilderMeshApproximate.cpp
--- a/Engine/Source/Developer/WorldPartitionHLODUtilities/Private/HLODBuilderMeshApproximate.cpp
+++ b/Engine/Source/Developer/WorldPartitionHLODUtilities/Private/HLODBuilderMeshApproximate.cpp
@@ -16,14 +16,39 @@
 #include "IGeometryProcessingInterfacesModule.h"
 #include "GeometryProcessingInterfaces/ApproximateActors.h"
 
+#include "HLODBuilderInstancing.h"
+
 
 TArray<UPrimitiveComponent*> FHLODBuilder_MeshApproximate::CreateComponents(AWorldPartitionHLOD* InHLODActor, const UHLODLayer* InHLODLayer, const TArray<UPrimitiveComponent*>& InSubComponents)
 {
 	TRACE_CPUPROFILER_EVENT_SCOPE(FHLODBuilder_MeshApproximate::CreateComponents);
 
 	TSet<AActor*> Actors;
-	Algo::Transform(InSubComponents, Actors, [](UPrimitiveComponent* InPrimitiveComponent) { return InPrimitiveComponent->GetOwner(); });
+	TArray<UPrimitiveComponent*> InstancedComponents;
+
+	// Filter the input components
+	for (UPrimitiveComponent* SubComponent : InSubComponents)
+	{
+		if (!SubComponent)
+		{
+			continue;
+		}
 
+		switch (SubComponent->HLODBatchingPolicy)
+		{
+		case EHLODBatchingPolicy::None:
+			Actors.Add(SubComponent->GetOwner());
+			break;
+		case EHLODBatchingPolicy::Instancing:
+			InstancedComponents.Add(SubComponent);
+			break;
+		case EHLODBatchingPolicy::MeshSection:
+			InstancedComponents.Add(SubComponent);
+			UE_LOG(LogHLODBuilder, Warning, TEXT("EHLODBatchingPolicy::MeshSection is not yet supported by the MeshApproximate builder."));
+			break;
+		}
+	}
+	
 	IGeometryProcessingInterfacesModule& GeomProcInterfaces = FModuleManager::Get().LoadModuleChecked<IGeometryProcessingInterfacesModule>("GeometryProcessingInterfaces");
 	IGeometryProcessing_ApproximateActors* ApproxActorsAPI = GeomProcInterfaces.GetApproximateActorsImplementation();
 
@@ -151,5 +176,12 @@ TArray<UPrimitiveComponent*> FHLODBuilder_MeshApproximate::CreateComponents(AWor
 		}
 	}
 
+	// Batch instances
+	if (InstancedComponents.Num())
+	{
+		FHLODBuilder_Instancing InstancingHLODBuilder;
+		Components.Append(InstancingHLODBuilder.CreateComponents(InHLODActor, nullptr, InstancedComponents));
+	}
+
 	return Components;
 }
diff --git a/Engine/Source/Developer/WorldPartitionHLODUtilities/Private/HLODBuilderMeshApproximate.h b/Engine/Source/Developer/WorldPartitionHLODUtilities/Private/HLODBuilderMeshApproximate.h
--- a/Engine/Source/Developer/WorldPartitionHLODUtilities/Private/HLODBuilderMeshApproximate.h
+++ b/Engine/Source/Developer/WorldPartitionHLODUtilities/Private/HLODBuilderMeshApproximate.h
@@ -9,5 +9,6 @@
  */
 class FHLODBuilder_MeshApproximate : public FHLODBuilder
 {
+public:
 	virtual TArray<UPrimitiveComponent*> CreateComponents(AWorldPartitionHLOD* InHLODActor, const UHLODLayer* InHLODLayer, const TArray<UPrimitiveComponent*>& InSubComponents) override;
 };
diff --git a/Engine/Source/Developer/WorldPartitionHLODUtilities/Private/HLODBuilderMeshMerge.h b/Engine/Source/Developer/WorldPartitionHLODUtilities/Private/HLODBuilderMeshMerge.h
--- a/Engine/Source/Developer/WorldPartitionHLODUtilities/Private/HLODBuilderMeshMerge.h
+++ b/Engine/Source/Developer/WorldPartitionHLODUtilities/Private/HLODBuilderMeshMerge.h
@@ -9,5 +9,6 @@
  */
 class FHLODBuilder_MeshMerge : public FHLODBuilder
 {
+public:
 	virtual TArray<UPrimitiveComponent*> CreateComponents(AWorldPartitionHLOD* InHLODActor, const UHLODLayer* InHLODLayer, const TArray<UPrimitiveComponent*>& InSubComponents) override;
 };
diff --git a/Engine/Source/Developer/WorldPartitionHLODUtilities/Private/HLODBuilderMeshSimplify.cpp b/Engine/Source/Developer/WorldPartitionHLODUtilities/Private/HLODBuilderMeshSimplify.cpp
--- a/Engine/Source/Developer/WorldPartitionHLODUtilities/Private/HLODBuilderMeshSimplify.cpp
+++ b/Engine/Source/Developer/WorldPartitionHLODUtilities/Private/HLODBuilderMeshSimplify.cpp
@@ -15,6 +15,8 @@
 #include "MeshMergeModule.h"
 #include "Modules/ModuleManager.h"
 
+#include "HLODBuilderInstancing.h"
+
 
 TArray<UPrimitiveComponent*> FHLODBuilder_MeshSimplify::CreateComponents(AWorldPartitionHLOD* InHLODActor, const UHLODLayer* InHLODLayer, const TArray<UPrimitiveComponent*>& InSubComponents)
 {
@@ -25,7 +27,33 @@ TArray<UPrimitiveComponent*> FHLODBuilder_MeshSimplify::CreateComponents(AWorldP
 	ProxyDelegate.BindLambda([&Assets](const FGuid Guid, TArray<UObject*>& InAssetsCreated) { Assets = InAssetsCreated; });
 
 	TArray<UStaticMeshComponent*> StaticMeshComponents;
-	Algo::TransformIf(InSubComponents, StaticMeshComponents, [](UPrimitiveComponent* InPrimitiveComponent) { return InPrimitiveComponent->IsA<UStaticMeshComponent>(); }, [](UPrimitiveComponent* InPrimitiveComponent) { return Cast<UStaticMeshComponent>(InPrimitiveComponent); });
+	TArray<UPrimitiveComponent*> InstancedComponents;
+
+	// Filter the input components
+	for (UPrimitiveComponent* SubComponent : InSubComponents)
+	{
+		if (!SubComponent)
+		{
+			continue;
+		}
+
+		switch (SubComponent->HLODBatchingPolicy)
+		{
+		case EHLODBatchingPolicy::None:
+			if (UStaticMeshComponent* SMC = Cast<UStaticMeshComponent>(SubComponent))
+			{
+				StaticMeshComponents.Add(SMC);
+			}
+			break;
+		case EHLODBatchingPolicy::Instancing:
+			InstancedComponents.Add(SubComponent);
+			break;
+		case EHLODBatchingPolicy::MeshSection:
+			InstancedComponents.Add(SubComponent);
+			UE_LOG(LogHLODBuilder, Warning, TEXT("EHLODBatchingPolicy::MeshSection is not yet supported by the MeshApproximate builder."));
+			break;
+		}
+	}
 
 	const IMeshMergeUtilities& MeshMergeUtilities = FModuleManager::Get().LoadModuleChecked<IMeshMergeModule>("MeshMergeUtilities").GetUtilities();
 	MeshMergeUtilities.CreateProxyMesh(StaticMeshComponents, InHLODLayer->GetMeshSimplifySettings(), InHLODLayer->GetHLODMaterial().LoadSynchronous(), InHLODActor->GetPackage(), InHLODActor->GetActorLabel(), FGuid::NewGuid(), ProxyDelegate, true);
@@ -43,5 +71,15 @@ TArray<UPrimitiveComponent*> FHLODBuilder_MeshSimplify::CreateComponents(AWorldP
 		}
 	});
 
-	return TArray<UPrimitiveComponent*>({ Component });
+	TArray<UPrimitiveComponent*> Components;
+	Components.Add(Component);
+
+	// Batch instances
+	if (InstancedComponents.Num())
+	{
+		FHLODBuilder_Instancing InstancingHLODBuilder;
+		Components.Append(InstancingHLODBuilder.CreateComponents(InHLODActor, nullptr, InstancedComponents));
+	}
+
+	return Components;
 }
diff --git a/Engine/Source/Developer/WorldPartitionHLODUtilities/Private/HLODBuilderMeshSimplify.h b/Engine/Source/Developer/WorldPartitionHLODUtilities/Private/HLODBuilderMeshSimplify.h
--- a/Engine/Source/Developer/WorldPartitionHLODUtilities/Private/HLODBuilderMeshSimplify.h
+++ b/Engine/Source/Developer/WorldPartitionHLODUtilities/Private/HLODBuilderMeshSimplify.h
@@ -9,5 +9,6 @@
  */
 class FHLODBuilder_MeshSimplify : public FHLODBuilder
 {
+public:
 	virtual TArray<UPrimitiveComponent*> CreateComponents(AWorldPartitionHLOD* InHLODActor, const UHLODLayer* InHLODLayer, const TArray<UPrimitiveComponent*>& InSubComponents) override;
 };
