HLODBuilderInstancing - Fixed issue with per instance custom data aggregation

diff --git a/Engine/Source/Developer/WorldPartitionHLODUtilities/Private/HLODBuilderInstancing.cpp b/Engine/Source/Developer/WorldPartitionHLODUtilities/Private/HLODBuilderInstancing.cpp
--- a/Engine/Source/Developer/WorldPartitionHLODUtilities/Private/HLODBuilderInstancing.cpp
+++ b/Engine/Source/Developer/WorldPartitionHLODUtilities/Private/HLODBuilderInstancing.cpp
@@ -29,11 +29,11 @@ namespace
 	// Store batched instances data
 	struct FInstancingData
 	{
-		int32							NumInstances;
+		int32							NumInstances = 0;
 
 		TArray<FTransform>				InstancesTransforms;
 
-		int32							NumCustomDataFloats;
+		int32							NumCustomDataFloats = 0;
 		TArray<float>					InstancesCustomData;
 	};
 }
@@ -55,7 +55,7 @@ TArray<UPrimitiveComponent*> FHLODBuilder_Instancing::CreateComponents(AWorldPar
 
 			if (UInstancedStaticMeshComponent* ISMC = Cast<UInstancedStaticMeshComponent>(SMC))
 			{
-				InstancingData.NumCustomDataFloats = FMath::Max(InstancingData.NumCustomDataFloats, ISMC->PerInstanceSMCustomData.Num());
+				InstancingData.NumCustomDataFloats = FMath::Max(InstancingData.NumCustomDataFloats, ISMC->NumCustomDataFloats);
 				InstancingData.NumInstances += ISMC->GetInstanceCount();
 			}
 			else
@@ -83,8 +83,6 @@ TArray<UPrimitiveComponent*> FHLODBuilder_Instancing::CreateComponents(AWorldPar
 			FCustomISMComponentDescriptor ISMComponentDescriptor(SMC);
 			FInstancingData& InstancingData = InstancesData.FindChecked(ISMComponentDescriptor);
 
-			int32 NumCustomDataFloatsAdded = 0;
-
 			if (UInstancedStaticMeshComponent* ISMC = Cast<UInstancedStaticMeshComponent>(SMC))
 			{
 				// Add transforms
@@ -95,16 +93,15 @@ TArray<UPrimitiveComponent*> FHLODBuilder_Instancing::CreateComponents(AWorldPar
 				}
 
 				// Add per instance custom data
+				int32 NumCustomDataFloatToAdd = ISMC->GetInstanceCount() * InstancingData.NumCustomDataFloats;
 				InstancingData.InstancesCustomData.Append(ISMC->PerInstanceSMCustomData);
-				NumCustomDataFloatsAdded = ISMC->PerInstanceSMCustomData.Num();
+				InstancingData.InstancesCustomData.AddDefaulted(NumCustomDataFloatToAdd - ISMC->PerInstanceSMCustomData.Num());
 			}
 			else
 			{
 				InstancingData.InstancesTransforms.Add(SMC->GetComponentTransform());
+				InstancingData.InstancesCustomData.AddDefaulted(InstancingData.NumCustomDataFloats);
 			}
-
-			// Add missing custom data, if any
-			InstancingData.InstancesCustomData.AddDefaulted(InstancingData.NumCustomDataFloats - NumCustomDataFloatsAdded);
 		}
 	}
 
@@ -114,6 +111,8 @@ TArray<UPrimitiveComponent*> FHLODBuilder_Instancing::CreateComponents(AWorldPar
 		const FISMComponentDescriptor& ISMComponentDescriptor = Entry.Key;
 		FInstancingData& EntryInstancingData = Entry.Value;
 
+		check(EntryInstancingData.InstancesTransforms.Num() * EntryInstancingData.NumCustomDataFloats == EntryInstancingData.InstancesCustomData.Num());
+
 		UInstancedStaticMeshComponent* Component = ISMComponentDescriptor.CreateComponent(InHLODActor);
 		Component->SetForcedLodModel(Component->GetStaticMesh()->GetNumLODs());
 		Component->NumCustomDataFloats = EntryInstancingData.NumCustomDataFloats;
