HLOD Builders - Refactored to add support for custom HLOD layer type.
* HLOD builder settings now specified through a base UHLODBuilderSettings class.

diff --git a/Engine/Source/Developer/WorldPartitionHLODUtilities/Private/HLODBuilder.cpp b/Engine/Source/Developer/WorldPartitionHLODUtilities/Private/HLODBuilder.cpp
deleted file mode 100644
diff --git a/Engine/Source/Developer/WorldPartitionHLODUtilities/Private/HLODBuilder.h b/Engine/Source/Developer/WorldPartitionHLODUtilities/Private/HLODBuilder.h
deleted file mode 100644
diff --git a/Engine/Source/Developer/WorldPartitionHLODUtilities/Private/HLODBuilderInstancing.cpp b/Engine/Source/Developer/WorldPartitionHLODUtilities/Private/HLODBuilderInstancing.cpp
--- a/Engine/Source/Developer/WorldPartitionHLODUtilities/Private/HLODBuilderInstancing.cpp
+++ b/Engine/Source/Developer/WorldPartitionHLODUtilities/Private/HLODBuilderInstancing.cpp
@@ -38,7 +38,12 @@ namespace
 	};
 }
 
-TArray<UPrimitiveComponent*> FHLODBuilder_Instancing::CreateComponents(AWorldPartitionHLOD* InHLODActor, const UHLODLayer* InHLODLayer, const TArray<UPrimitiveComponent*>& InSubComponents)
+UHLODBuilderInstancing::UHLODBuilderInstancing(const FObjectInitializer& ObjectInitializer)
+	: Super(ObjectInitializer)
+{
+}
+
+TArray<UPrimitiveComponent*> UHLODBuilderInstancing::CreateComponents(AWorldPartitionHLOD* InHLODActor, const UHLODLayer* InHLODLayer, const TArray<UPrimitiveComponent*>& InSubComponents) const
 {
 	TRACE_CPUPROFILER_EVENT_SCOPE(FHLODBuilder_Instancing::CreateComponents);
 
diff --git a/Engine/Source/Developer/WorldPartitionHLODUtilities/Private/HLODBuilderInstancing.h b/Engine/Source/Developer/WorldPartitionHLODUtilities/Private/HLODBuilderInstancing.h
--- a/Engine/Source/Developer/WorldPartitionHLODUtilities/Private/HLODBuilderInstancing.h
+++ b/Engine/Source/Developer/WorldPartitionHLODUtilities/Private/HLODBuilderInstancing.h
@@ -2,15 +2,20 @@
 
 #pragma once
 
-#include "HLODBuilder.h"
+#include "WorldPartition/HLOD/HLODBuilder.h"
+#include "HLODBuilderInstancing.generated.h"
+
 
 /**
  * Build a AWorldPartitionHLOD whose components are ISMC
  */
-class FHLODBuilder_Instancing : public FHLODBuilder
+UCLASS()
+class UHLODBuilderInstancing : public UHLODBuilder
 {
+	 GENERATED_UCLASS_BODY()
+
 public:
 	virtual bool RequiresCompiledAssets() const override { return false; }
 
-	virtual TArray<UPrimitiveComponent*> CreateComponents(AWorldPartitionHLOD* InHLODActor, const UHLODLayer* InHLODLayer, const TArray<UPrimitiveComponent*>& InSubComponents) override;
+	virtual TArray<UPrimitiveComponent*> CreateComponents(AWorldPartitionHLOD* InHLODActor, const UHLODLayer* InHLODLayer, const TArray<UPrimitiveComponent*>& InSubComponents) const override;
 };
diff --git a/Engine/Source/Developer/WorldPartitionHLODUtilities/Private/HLODBuilderMeshApproximate.cpp b/Engine/Source/Developer/WorldPartitionHLODUtilities/Private/HLODBuilderMeshApproximate.cpp
--- a/Engine/Source/Developer/WorldPartitionHLODUtilities/Private/HLODBuilderMeshApproximate.cpp
+++ b/Engine/Source/Developer/WorldPartitionHLODUtilities/Private/HLODBuilderMeshApproximate.cpp
@@ -16,10 +16,69 @@
 #include "IGeometryProcessingInterfacesModule.h"
 #include "GeometryProcessingInterfaces/ApproximateActors.h"
 
+#include "Materials/Material.h"
+#include "Engine/HLODProxy.h"
+#include "Serialization/ArchiveCrc32.h"
+
 #include "HLODBuilderInstancing.h"
 
 
-TArray<UPrimitiveComponent*> FHLODBuilder_MeshApproximate::CreateComponents(AWorldPartitionHLOD* InHLODActor, const UHLODLayer* InHLODLayer, const TArray<UPrimitiveComponent*>& InSubComponents)
+UHLODBuilderMeshApproximate::UHLODBuilderMeshApproximate(const FObjectInitializer& ObjectInitializer)
+	: Super(ObjectInitializer)
+{
+}
+
+UHLODBuilderMeshApproximateSettings::UHLODBuilderMeshApproximateSettings(const FObjectInitializer& ObjectInitializer)
+	: Super(ObjectInitializer)
+{
+#if WITH_EDITORONLY_DATA
+	if (!IsTemplate())
+	{
+		HLODMaterial = GEngine->DefaultHLODFlattenMaterial;
+	}
+#endif
+}
+
+uint32 UHLODBuilderMeshApproximateSettings::GetCRC() const
+{
+	UHLODBuilderMeshApproximateSettings& This = *const_cast<UHLODBuilderMeshApproximateSettings*>(this);
+
+	FArchiveCrc32 Ar;
+
+	Ar << This.MeshApproximationSettings;
+	UE_LOG(LogHLODBuilder, VeryVerbose, TEXT(" - MeshApproximationSettings = %d"), Ar.GetCrc());
+
+	uint32 Hash = Ar.GetCrc();
+
+	if (!HLODMaterial.IsNull())
+	{
+		UMaterialInterface* Material = HLODMaterial.LoadSynchronous();
+		if (Material)
+		{
+			uint32 MaterialCRC = UHLODProxy::GetCRC(Material);
+			UE_LOG(LogHLODBuilder, VeryVerbose, TEXT(" - Material = %d"), MaterialCRC);
+			Hash = HashCombine(Hash, MaterialCRC);
+		}
+	}
+
+	return Hash;
+}
+
+UHLODBuilderSettings* UHLODBuilderMeshApproximate::CreateSettings(UHLODLayer* InHLODLayer) const 
+{
+	UHLODBuilderMeshApproximateSettings* HLODBuilderSettings = NewObject<UHLODBuilderMeshApproximateSettings>(InHLODLayer);
+
+	// If previous settings object is null, this means we have an older version of the object. Populate with the deprecated settings.
+	if (InHLODLayer->GetHLODBuilderSettings() == nullptr)
+	{
+		HLODBuilderSettings->MeshApproximationSettings = InHLODLayer->MeshApproximationSettings_DEPRECATED;
+		HLODBuilderSettings->HLODMaterial = InHLODLayer->HLODMaterial_DEPRECATED;
+	}
+
+	return HLODBuilderSettings;
+}
+
+TArray<UPrimitiveComponent*> UHLODBuilderMeshApproximate::CreateComponents(AWorldPartitionHLOD* InHLODActor, const UHLODLayer* InHLODLayer, const TArray<UPrimitiveComponent*>& InSubComponents) const
 {
 	TRACE_CPUPROFILER_EVENT_SCOPE(FHLODBuilder_MeshApproximate::CreateComponents);
 
@@ -56,7 +115,9 @@ TArray<UPrimitiveComponent*> FHLODBuilder_MeshApproximate::CreateComponents(AWor
 	// Construct options for ApproximateActors operation
 	//
 
-	const FMeshApproximationSettings& UseSettings = InHLODLayer->GetMeshApproximationSettings();
+	const UHLODBuilderMeshApproximateSettings* MeshApproximateSettings = CastChecked<UHLODBuilderMeshApproximateSettings>(InHLODLayer->GetHLODBuilderSettings());
+	const FMeshApproximationSettings& UseSettings = MeshApproximateSettings->MeshApproximationSettings;
+	UMaterial* HLODMaterial = MeshApproximateSettings->HLODMaterial.LoadSynchronous();
 
 	IGeometryProcessing_ApproximateActors::FOptions Options = ApproxActorsAPI->ConstructOptions(UseSettings);
 	Options.BasePackagePath = InHLODActor->GetPackage()->GetName();
@@ -64,7 +125,7 @@ TArray<UPrimitiveComponent*> FHLODBuilder_MeshApproximate::CreateComponents(AWor
 	Options.bCreatePhysicsBody = false;
 
 	// Material baking settings
-	Options.BakeMaterial = GEngine->DefaultHLODFlattenMaterial;
+	Options.BakeMaterial = HLODMaterial;
 	Options.BaseColorTexParamName = FName("BaseColorTexture");
 	Options.NormalTexParamName = FName("NormalTexture");
 	Options.MetallicTexParamName = FName("MetallicTexture");
@@ -178,8 +239,8 @@ TArray<UPrimitiveComponent*> FHLODBuilder_MeshApproximate::CreateComponents(AWor
 	// Batch instances
 	if (InstancedComponents.Num())
 	{
-		FHLODBuilder_Instancing InstancingHLODBuilder;
-		Components.Append(InstancingHLODBuilder.CreateComponents(InHLODActor, nullptr, InstancedComponents));
+		UHLODBuilderInstancing* InstancingHLODBuilder = NewObject<UHLODBuilderInstancing>();
+		Components.Append(InstancingHLODBuilder->CreateComponents(InHLODActor, nullptr, InstancedComponents));
 	}
 
 	return Components;
diff --git a/Engine/Source/Developer/WorldPartitionHLODUtilities/Private/HLODBuilderMeshApproximate.h b/Engine/Source/Developer/WorldPartitionHLODUtilities/Private/HLODBuilderMeshApproximate.h
--- a/Engine/Source/Developer/WorldPartitionHLODUtilities/Private/HLODBuilderMeshApproximate.h
+++ b/Engine/Source/Developer/WorldPartitionHLODUtilities/Private/HLODBuilderMeshApproximate.h
@@ -2,13 +2,40 @@
 
 #pragma once
 
-#include "HLODBuilder.h"
+#include "WorldPartition/HLOD/HLODBuilder.h"
+#include "Engine/MeshMerging.h"
+
+#include "HLODBuilderMeshApproximate.generated.h"
+
+class UMaterial;
+
+
+UCLASS(Blueprintable, Config = Engine, PerObjectConfig)
+class UHLODBuilderMeshApproximateSettings : public UHLODBuilderSettings
+{
+	GENERATED_UCLASS_BODY()
+
+	virtual uint32 GetCRC() const override;
+
+	/** Mesh approximation settings */
+	UPROPERTY(EditAnywhere, Config, Category = HLOD, meta=(EditInline))
+	FMeshApproximationSettings MeshApproximationSettings;
+
+	/** Material that will be used by the generated HLOD static mesh */
+	UPROPERTY(EditAnywhere, Config, AdvancedDisplay, Category = HLOD)
+	TSoftObjectPtr<UMaterial> HLODMaterial;
+};
+
 
 /**
  * Build an approximated mesh using geometry from the provided actors
  */
-class FHLODBuilder_MeshApproximate : public FHLODBuilder
+UCLASS()
+class UHLODBuilderMeshApproximate : public UHLODBuilder
 {
+	GENERATED_UCLASS_BODY()
+
 public:
-	virtual TArray<UPrimitiveComponent*> CreateComponents(AWorldPartitionHLOD* InHLODActor, const UHLODLayer* InHLODLayer, const TArray<UPrimitiveComponent*>& InSubComponents) override;
+	virtual UHLODBuilderSettings* CreateSettings(UHLODLayer* InHLODLayer) const override;
+	virtual TArray<UPrimitiveComponent*> CreateComponents(AWorldPartitionHLOD* InHLODActor, const UHLODLayer* InHLODLayer, const TArray<UPrimitiveComponent*>& InSubComponents) const override;
 };
diff --git a/Engine/Source/Developer/WorldPartitionHLODUtilities/Private/HLODBuilderMeshMerge.cpp b/Engine/Source/Developer/WorldPartitionHLODUtilities/Private/HLODBuilderMeshMerge.cpp
--- a/Engine/Source/Developer/WorldPartitionHLODUtilities/Private/HLODBuilderMeshMerge.cpp
+++ b/Engine/Source/Developer/WorldPartitionHLODUtilities/Private/HLODBuilderMeshMerge.cpp
@@ -14,16 +14,79 @@
 #include "MeshMergeModule.h"
 #include "Modules/ModuleManager.h"
 
+#include "Materials/Material.h"
+#include "Engine/HLODProxy.h"
+#include "Serialization/ArchiveCrc32.h"
+
+
+UHLODBuilderMeshMerge::UHLODBuilderMeshMerge(const FObjectInitializer& ObjectInitializer)
+	: Super(ObjectInitializer)
+{
+}
+
+UHLODBuilderMeshMergeSettings::UHLODBuilderMeshMergeSettings(const FObjectInitializer& ObjectInitializer)
+	: Super(ObjectInitializer)
+{
+#if WITH_EDITORONLY_DATA
+	if (!IsTemplate())
+	{
+		HLODMaterial = GEngine->DefaultHLODFlattenMaterial;
+	}
+#endif
+}
 
-TArray<UPrimitiveComponent*> FHLODBuilder_MeshMerge::CreateComponents(AWorldPartitionHLOD* InHLODActor, const UHLODLayer* InHLODLayer, const TArray<UPrimitiveComponent*>& InSubComponents)
+uint32 UHLODBuilderMeshMergeSettings::GetCRC() const
+{
+	UHLODBuilderMeshMergeSettings& This = *const_cast<UHLODBuilderMeshMergeSettings*>(this);
+
+	FArchiveCrc32 Ar;
+
+	Ar << This.MeshMergeSettings;
+	UE_LOG(LogHLODBuilder, VeryVerbose, TEXT(" - MeshMergeSettings = %d"), Ar.GetCrc());
+
+	uint32 Hash = Ar.GetCrc();
+
+	if (!HLODMaterial.IsNull())
+	{
+		UMaterialInterface* Material = HLODMaterial.LoadSynchronous();
+		if (Material)
+		{
+			uint32 MaterialCRC = UHLODProxy::GetCRC(Material);
+			UE_LOG(LogHLODBuilder, VeryVerbose, TEXT(" - Material = %d"), MaterialCRC);
+			Hash = HashCombine(Hash, MaterialCRC);
+		}
+	}
+
+	return Hash;
+}
+
+UHLODBuilderSettings* UHLODBuilderMeshMerge::CreateSettings(UHLODLayer* InHLODLayer) const
+{
+	UHLODBuilderMeshMergeSettings* HLODBuilderSettings = NewObject<UHLODBuilderMeshMergeSettings>(InHLODLayer);
+
+	// If previous settings object is null, this means we have an older version of the object. Populate with the deprecated settings.
+	if (InHLODLayer->GetHLODBuilderSettings() == nullptr)
+	{
+		HLODBuilderSettings->MeshMergeSettings = InHLODLayer->MeshMergeSettings_DEPRECATED;
+		HLODBuilderSettings->HLODMaterial = InHLODLayer->HLODMaterial_DEPRECATED;
+	}
+
+	return HLODBuilderSettings;
+}
+
+TArray<UPrimitiveComponent*> UHLODBuilderMeshMerge::CreateComponents(AWorldPartitionHLOD* InHLODActor, const UHLODLayer* InHLODLayer, const TArray<UPrimitiveComponent*>& InSubComponents) const
 {
 	TRACE_CPUPROFILER_EVENT_SCOPE(FHLODBuilder_MeshMerge::CreateComponents);
 
 	TArray<UObject*> Assets;
 	FVector MergedActorLocation;
 
+	const UHLODBuilderMeshMergeSettings* MeshMergeSettings = CastChecked<UHLODBuilderMeshMergeSettings>(InHLODLayer->GetHLODBuilderSettings());
+	const FMeshMergingSettings& UseSettings = MeshMergeSettings->MeshMergeSettings;
+	UMaterial* HLODMaterial = MeshMergeSettings->HLODMaterial.LoadSynchronous();
+
 	const IMeshMergeUtilities& MeshMergeUtilities = FModuleManager::Get().LoadModuleChecked<IMeshMergeModule>("MeshMergeUtilities").GetUtilities();
-	MeshMergeUtilities.MergeComponentsToStaticMesh(InSubComponents, InHLODActor->GetWorld(), InHLODLayer->GetMeshMergeSettings(), InHLODLayer->GetHLODMaterial().LoadSynchronous(), InHLODActor->GetPackage(), InHLODActor->GetActorLabel(), Assets, MergedActorLocation, 0.25f, false);
+	MeshMergeUtilities.MergeComponentsToStaticMesh(InSubComponents, InHLODActor->GetWorld(), UseSettings, HLODMaterial, InHLODActor->GetPackage(), InHLODActor->GetActorLabel(), Assets, MergedActorLocation, 0.25f, false);
 
 	UStaticMeshComponent* Component = nullptr;
 	Algo::ForEach(Assets, [this, InHLODActor, &Component, &MergedActorLocation](UObject* Asset)
diff --git a/Engine/Source/Developer/WorldPartitionHLODUtilities/Private/HLODBuilderMeshMerge.h b/Engine/Source/Developer/WorldPartitionHLODUtilities/Private/HLODBuilderMeshMerge.h
--- a/Engine/Source/Developer/WorldPartitionHLODUtilities/Private/HLODBuilderMeshMerge.h
+++ b/Engine/Source/Developer/WorldPartitionHLODUtilities/Private/HLODBuilderMeshMerge.h
@@ -2,13 +2,40 @@
 
 #pragma once
 
-#include "HLODBuilder.h"
+#include "WorldPartition/HLOD/HLODBuilder.h"
+#include "Engine/MeshMerging.h"
+
+#include "HLODBuilderMeshMerge.generated.h"
+
+class UMaterial;
+
+
+UCLASS(Blueprintable, Config = Engine, PerObjectConfig)
+class UHLODBuilderMeshMergeSettings : public UHLODBuilderSettings
+{
+	GENERATED_UCLASS_BODY()
+
+	virtual uint32 GetCRC() const override;
+
+	/** Merged mesh generation settings */
+	UPROPERTY(EditAnywhere, Config, Category = HLOD)
+	FMeshMergingSettings MeshMergeSettings;
+
+	/** Material that will be used by the generated HLOD static mesh */
+	UPROPERTY(EditAnywhere, Config, AdvancedDisplay, Category = HLOD)
+	TSoftObjectPtr<UMaterial> HLODMaterial;
+};
+
 
 /**
  * Build a merged mesh using geometry from the provided actors
  */
-class FHLODBuilder_MeshMerge : public FHLODBuilder
+UCLASS()
+class UHLODBuilderMeshMerge : public UHLODBuilder
 {
+	GENERATED_UCLASS_BODY()
+
 public:
-	virtual TArray<UPrimitiveComponent*> CreateComponents(AWorldPartitionHLOD* InHLODActor, const UHLODLayer* InHLODLayer, const TArray<UPrimitiveComponent*>& InSubComponents) override;
+	virtual UHLODBuilderSettings* CreateSettings(UHLODLayer* InHLODLayer) const override;
+	virtual TArray<UPrimitiveComponent*> CreateComponents(AWorldPartitionHLOD* InHLODActor, const UHLODLayer* InHLODLayer, const TArray<UPrimitiveComponent*>& InSubComponents) const override;
 };
diff --git a/Engine/Source/Developer/WorldPartitionHLODUtilities/Private/HLODBuilderMeshSimplify.cpp b/Engine/Source/Developer/WorldPartitionHLODUtilities/Private/HLODBuilderMeshSimplify.cpp
--- a/Engine/Source/Developer/WorldPartitionHLODUtilities/Private/HLODBuilderMeshSimplify.cpp
+++ b/Engine/Source/Developer/WorldPartitionHLODUtilities/Private/HLODBuilderMeshSimplify.cpp
@@ -15,10 +15,69 @@
 #include "MeshMergeModule.h"
 #include "Modules/ModuleManager.h"
 
+#include "Materials/Material.h"
+#include "Engine/HLODProxy.h"
+#include "Serialization/ArchiveCrc32.h"
+
 #include "HLODBuilderInstancing.h"
 
 
-TArray<UPrimitiveComponent*> FHLODBuilder_MeshSimplify::CreateComponents(AWorldPartitionHLOD* InHLODActor, const UHLODLayer* InHLODLayer, const TArray<UPrimitiveComponent*>& InSubComponents)
+UHLODBuilderMeshSimplify::UHLODBuilderMeshSimplify(const FObjectInitializer& ObjectInitializer)
+	: Super(ObjectInitializer)
+{
+}
+
+UHLODBuilderMeshSimplifySettings::UHLODBuilderMeshSimplifySettings(const FObjectInitializer& ObjectInitializer)
+	: Super(ObjectInitializer)
+{
+#if WITH_EDITORONLY_DATA
+	if (!IsTemplate())
+	{
+		HLODMaterial = GEngine->DefaultHLODFlattenMaterial;
+	}
+#endif
+}
+
+uint32 UHLODBuilderMeshSimplifySettings::GetCRC() const
+{
+	UHLODBuilderMeshSimplifySettings& This = *const_cast<UHLODBuilderMeshSimplifySettings*>(this);
+
+	FArchiveCrc32 Ar;
+
+	Ar << This.MeshSimplifySettings;
+	UE_LOG(LogHLODBuilder, VeryVerbose, TEXT(" - MeshSimplifySettings = %d"), Ar.GetCrc());
+
+	uint32 Hash = Ar.GetCrc();
+
+	if (!HLODMaterial.IsNull())
+	{
+		UMaterialInterface* Material = HLODMaterial.LoadSynchronous();
+		if (Material)
+		{
+			uint32 MaterialCRC = UHLODProxy::GetCRC(Material);
+			UE_LOG(LogHLODBuilder, VeryVerbose, TEXT(" - Material = %d"), MaterialCRC);
+			Hash = HashCombine(Hash, MaterialCRC);
+		}
+	}
+
+	return Hash;
+}
+
+UHLODBuilderSettings* UHLODBuilderMeshSimplify::CreateSettings(UHLODLayer* InHLODLayer) const
+{
+	UHLODBuilderMeshSimplifySettings* HLODBuilderSettings = NewObject<UHLODBuilderMeshSimplifySettings>(InHLODLayer);
+
+	// If previous settings object is null, this means we have an older version of the object. Populate with the deprecated settings.
+	if (InHLODLayer->GetHLODBuilderSettings() == nullptr)
+	{
+		HLODBuilderSettings->MeshSimplifySettings = InHLODLayer->MeshSimplifySettings_DEPRECATED;
+		HLODBuilderSettings->HLODMaterial = InHLODLayer->HLODMaterial_DEPRECATED;
+	}
+
+	return HLODBuilderSettings;
+}
+
+TArray<UPrimitiveComponent*> UHLODBuilderMeshSimplify::CreateComponents(AWorldPartitionHLOD* InHLODActor, const UHLODLayer* InHLODLayer, const TArray<UPrimitiveComponent*>& InSubComponents) const
 {
 	TRACE_CPUPROFILER_EVENT_SCOPE(FHLODBuilder_MeshSimplify::CreateComponents);
 
@@ -55,8 +114,12 @@ TArray<UPrimitiveComponent*> FHLODBuilder_MeshSimplify::CreateComponents(AWorldP
 		}
 	}
 
+	const UHLODBuilderMeshSimplifySettings* MeshSimplifySettings = CastChecked<UHLODBuilderMeshSimplifySettings>(InHLODLayer->GetHLODBuilderSettings());
+	const FMeshProxySettings& UseSettings = MeshSimplifySettings->MeshSimplifySettings;
+	UMaterial* HLODMaterial = MeshSimplifySettings->HLODMaterial.LoadSynchronous();
+
 	const IMeshMergeUtilities& MeshMergeUtilities = FModuleManager::Get().LoadModuleChecked<IMeshMergeModule>("MeshMergeUtilities").GetUtilities();
-	MeshMergeUtilities.CreateProxyMesh(StaticMeshComponents, InHLODLayer->GetMeshSimplifySettings(), InHLODLayer->GetHLODMaterial().LoadSynchronous(), InHLODActor->GetPackage(), InHLODActor->GetActorLabel(), FGuid::NewGuid(), ProxyDelegate, true);
+	MeshMergeUtilities.CreateProxyMesh(StaticMeshComponents, UseSettings, HLODMaterial, InHLODActor->GetPackage(), InHLODActor->GetActorLabel(), FGuid::NewGuid(), ProxyDelegate, true);
 
 	UStaticMeshComponent* Component = nullptr;
 	Algo::ForEach(Assets, [this, InHLODActor, &Component](UObject* Asset)
@@ -76,8 +139,8 @@ TArray<UPrimitiveComponent*> FHLODBuilder_MeshSimplify::CreateComponents(AWorldP
 	// Batch instances
 	if (InstancedComponents.Num())
 	{
-		FHLODBuilder_Instancing InstancingHLODBuilder;
-		Components.Append(InstancingHLODBuilder.CreateComponents(InHLODActor, nullptr, InstancedComponents));
+		UHLODBuilderInstancing* InstancingHLODBuilder = NewObject<UHLODBuilderInstancing>();
+		Components.Append(InstancingHLODBuilder->CreateComponents(InHLODActor, nullptr, InstancedComponents));
 	}
 
 	return Components;
diff --git a/Engine/Source/Developer/WorldPartitionHLODUtilities/Private/HLODBuilderMeshSimplify.h b/Engine/Source/Developer/WorldPartitionHLODUtilities/Private/HLODBuilderMeshSimplify.h
--- a/Engine/Source/Developer/WorldPartitionHLODUtilities/Private/HLODBuilderMeshSimplify.h
+++ b/Engine/Source/Developer/WorldPartitionHLODUtilities/Private/HLODBuilderMeshSimplify.h
@@ -2,13 +2,40 @@
 
 #pragma once
 
-#include "HLODBuilder.h"
+#include "WorldPartition/HLOD/HLODBuilder.h"
+#include "Engine/MeshMerging.h"
+
+#include "HLODBuilderMeshSimplify.generated.h"
+
+class UMaterial;
+
+
+UCLASS(Blueprintable, Config = Engine, PerObjectConfig)
+class UHLODBuilderMeshSimplifySettings : public UHLODBuilderSettings
+{
+	GENERATED_UCLASS_BODY()
+
+	virtual uint32 GetCRC() const override;
+
+	/** Simplified mesh generation settings */
+	UPROPERTY(EditAnywhere, Config, Category = HLOD)
+	FMeshProxySettings MeshSimplifySettings;
+
+	/** Material that will be used by the generated HLOD static mesh */
+	UPROPERTY(EditAnywhere, Config, AdvancedDisplay, Category = HLOD)
+	TSoftObjectPtr<UMaterial> HLODMaterial;
+};
+
 
 /**
  * Build a simplified mesh using geometry from the provided actors
  */
-class FHLODBuilder_MeshSimplify : public FHLODBuilder
+UCLASS()
+class UHLODBuilderMeshSimplify : public UHLODBuilder
 {
+	GENERATED_UCLASS_BODY()
+
 public:
-	virtual TArray<UPrimitiveComponent*> CreateComponents(AWorldPartitionHLOD* InHLODActor, const UHLODLayer* InHLODLayer, const TArray<UPrimitiveComponent*>& InSubComponents) override;
+	virtual UHLODBuilderSettings* CreateSettings(UHLODLayer* InHLODLayer) const override;
+	virtual TArray<UPrimitiveComponent*> CreateComponents(AWorldPartitionHLOD* InHLODActor, const UHLODLayer* InHLODLayer, const TArray<UPrimitiveComponent*>& InSubComponents) const override;
 };
diff --git a/Engine/Source/Developer/WorldPartitionHLODUtilities/Private/WorldPartitionHLODUtilities.cpp b/Engine/Source/Developer/WorldPartitionHLODUtilities/Private/WorldPartitionHLODUtilities.cpp
--- a/Engine/Source/Developer/WorldPartitionHLODUtilities/Private/WorldPartitionHLODUtilities.cpp
+++ b/Engine/Source/Developer/WorldPartitionHLODUtilities/Private/WorldPartitionHLODUtilities.cpp
@@ -7,8 +7,9 @@
 #include "WorldPartition/WorldPartition.h"
 #include "WorldPartition/WorldPartitionHandle.h"
 #include "WorldPartition/HLOD/HLODActor.h"
-#include "WorldPartition/HLOD/HLODLayer.h"
 #include "WorldPartition/HLOD/HLODActorDesc.h"
+#include "WorldPartition/HLOD/HLODBuilder.h"
+#include "WorldPartition/HLOD/HLODLayer.h"
 
 #include "ISMPartition/ISMComponentDescriptor.h"
 
@@ -17,13 +18,13 @@
 
 #include "Engine/StaticMesh.h"
 #include "Engine/HLODProxy.h"
+#include "Serialization/ArchiveCrc32.h"
 #include "Materials/Material.h"
 #include "AssetCompilingManager.h"
 
 #include "LevelInstance/LevelInstanceActor.h"
 #include "LevelInstance/LevelInstanceSubsystem.h"
 
-#include "HLODBuilder.h"
 #include "HLODBuilderInstancing.h"
 #include "HLODBuilderMeshMerge.h"
 #include "HLODBuilderMeshSimplify.h"
@@ -69,6 +70,24 @@ static bool LoadSubActors(AWorldPartitionHLOD* InHLODActor, TArray<FWorldPartiti
 	return !bIsDirty;
 }
 
+static uint32 GetCRC(const UHLODLayer* InHLODLayer)
+{
+	UHLODLayer& HLODLayer= *const_cast<UHLODLayer*>(InHLODLayer);
+
+	uint32 CRC;
+
+	CRC = GetTypeHash(HLODLayer.GetLayerType());
+	UE_LOG(LogHLODBuilder, VeryVerbose, TEXT(" - LayerType = %d"), CRC);
+
+	CRC = HashCombine(HLODLayer.GetHLODBuilderSettings()->GetCRC(), CRC);
+	UE_LOG(LogHLODBuilder, VeryVerbose, TEXT(" - HLODBuilderSettings = %d"), CRC);
+
+	CRC = HashCombine(HLODLayer.GetCellSize(), CRC);
+	UE_LOG(LogHLODBuilder, VeryVerbose, TEXT(" - CellSize = %d"), CRC);
+
+	return CRC;
+}
+
 static uint32 ComputeHLODHash(AWorldPartitionHLOD* InHLODActor, const TArray<FWorldPartitionReference>& InActors)
 {
 	FArchiveCrc32 Ar;
@@ -78,13 +97,13 @@ static uint32 ComputeHLODHash(AWorldPartitionHLOD* InHLODActor, const TArray<FWo
 	Ar << HLODBaseKey;
 
 	// HLOD Layer
-	uint32 HLODLayerHash = InHLODActor->GetSubActorsHLODLayer()->GetCRC();
+	uint32 HLODLayerHash = GetCRC(InHLODActor->GetSubActorsHLODLayer());
 	UE_LOG(LogHLODBuilder, VeryVerbose, TEXT(" - HLODLayer (%s) = %x"), *InHLODActor->GetSubActorsHLODLayer()->GetName(), HLODLayerHash);
 	Ar << HLODLayerHash;
 
 	// We get the CRC of each component
 	TArray<uint32> ComponentsCRCs;
-	for (UPrimitiveComponent* Component : FHLODBuilder::GatherPrimitiveComponents(InActors))
+	for (UPrimitiveComponent* Component : UHLODBuilder::GatherPrimitiveComponents(InActors))
 	{
 		if (UStaticMeshComponent* StaticMeshComponent = Cast<UStaticMeshComponent>(Component))
 		{
@@ -293,6 +312,54 @@ TArray<AWorldPartitionHLOD*> FWorldPartitionHLODUtilities::CreateHLODActors(FHLO
 	return HLODActors;
 }
 
+static TSubclassOf<UHLODBuilder> GetHLODBuilderClassForLayer(const UHLODLayer* InHLODLayer)
+{
+	EHLODLayerType HLODLayerType = InHLODLayer->GetLayerType();
+	switch (HLODLayerType)
+	{
+	case EHLODLayerType::Instancing:
+		return UHLODBuilderInstancing::StaticClass();
+		break;
+
+	case EHLODLayerType::MeshMerge:
+		return UHLODBuilderMeshMerge::StaticClass();
+		break;
+
+	case EHLODLayerType::MeshSimplify:
+		return UHLODBuilderMeshSimplify::StaticClass();
+		break;
+
+	case EHLODLayerType::MeshApproximate:
+		return UHLODBuilderMeshApproximate::StaticClass();
+		break;
+
+	case EHLODLayerType::Custom:
+		return InHLODLayer->GetHLODBuilderClass();
+		break;
+
+	default:
+		checkf(false, TEXT("Unsupported type"));
+		return nullptr;
+	}
+}
+
+UHLODBuilderSettings* FWorldPartitionHLODUtilities::CreateHLODBuilderSettings(UHLODLayer* InHLODLayer)
+{
+	TSubclassOf<UHLODBuilder> HLODBuilderClass = GetHLODBuilderClassForLayer(InHLODLayer);
+	if (!HLODBuilderClass)
+	{
+		return NewObject<UHLODBuilderSettings>(InHLODLayer, UHLODBuilderSettings::StaticClass());
+	}
+
+	UHLODBuilderSettings* HLODBuilderSettings = HLODBuilderClass->GetDefaultObject<UHLODBuilder>()->CreateSettings(InHLODLayer);
+	if (!ensure(HLODBuilderSettings))
+	{
+		return NewObject<UHLODBuilderSettings>(InHLODLayer, UHLODBuilderSettings::StaticClass());
+	}
+
+	return HLODBuilderSettings;
+}
+
 uint32 FWorldPartitionHLODUtilities::BuildHLOD(AWorldPartitionHLOD* InHLODActor)
 {
 	TArray<FWorldPartitionReference> SubActors;
@@ -311,44 +378,27 @@ uint32 FWorldPartitionHLODUtilities::BuildHLOD(AWorldPartitionHLOD* InHLODActor)
 		return OldHLODHash;
 	}
 
-	TUniquePtr<FHLODBuilder> HLODBuilder = nullptr;
-	
 	const UHLODLayer* HLODLayer = InHLODActor->GetSubActorsHLODLayer();
-	EHLODLayerType HLODLevelType = HLODLayer->GetLayerType();
-	switch (HLODLevelType)
-	{
-	case EHLODLayerType::Instancing:
-		HLODBuilder = TUniquePtr<FHLODBuilder>(new FHLODBuilder_Instancing());
-		break;
-
-	case EHLODLayerType::MeshMerge:
-		HLODBuilder = TUniquePtr<FHLODBuilder>(new FHLODBuilder_MeshMerge());
-		break;
+	TSubclassOf<UHLODBuilder> HLODBuilderClass = GetHLODBuilderClassForLayer(HLODLayer);
 
-	case EHLODLayerType::MeshSimplify:
-		HLODBuilder = TUniquePtr<FHLODBuilder>(new FHLODBuilder_MeshSimplify());
-		break;
-
-	case EHLODLayerType::MeshApproximate:
-		HLODBuilder = TUniquePtr<FHLODBuilder>(new FHLODBuilder_MeshApproximate());
-		break;
-
-	default:
-		checkf(false, TEXT("Unsupported type"));
-	}
-
-	if (ensure(HLODBuilder))
+	if (HLODBuilderClass)
 	{
-		if (HLODBuilder->RequiresCompiledAssets())
+		UHLODBuilder* HLODBuilder = NewObject<UHLODBuilder>(GetTransientPackage(), HLODBuilderClass);
+		if (ensure(HLODBuilder))
 		{
-			// Wait for compilation to finish
-			FAssetCompilingManager::Get().FinishAllCompilation();
-		}
+			HLODBuilder->AddToRoot();
+			if (HLODBuilder->RequiresCompiledAssets())
+			{
+				// Wait for compilation to finish
+				FAssetCompilingManager::Get().FinishAllCompilation();
+			}
 
-		HLODBuilder->Build(InHLODActor, HLODLayer, SubActors);
-	}
+			HLODBuilder->Build(InHLODActor, HLODLayer, SubActors);
+			HLODBuilder->RemoveFromRoot();
+		}
 			
-	InHLODActor->MarkPackageDirty();
+		InHLODActor->MarkPackageDirty();
+	}
 
 	return NewHLODHash;
 }
diff --git a/Engine/Source/Developer/WorldPartitionHLODUtilities/Public/IWorldPartitionHLODUtilities.h b/Engine/Source/Developer/WorldPartitionHLODUtilities/Public/IWorldPartitionHLODUtilities.h
--- a/Engine/Source/Developer/WorldPartitionHLODUtilities/Public/IWorldPartitionHLODUtilities.h
+++ b/Engine/Source/Developer/WorldPartitionHLODUtilities/Public/IWorldPartitionHLODUtilities.h
@@ -11,6 +11,7 @@
 class AActor;
 class UWorldPartition;
 class UHLODLayer;
+class UHLODBuilderSettings;
 class AWorldPartitionHLOD;
 
 struct ENGINE_API FHLODCreationContext
@@ -60,4 +61,12 @@ public:
 	 * @return An hash that represent the content used to build this HLOD.
 	 */
 	virtual uint32 BuildHLOD(AWorldPartitionHLOD* InHLODActor) = 0;
+
+	/**
+	 * Create the HLOD builder settings for the provided HLOD layer object. The type of settings created will depend on the HLOD layer type.
+	 *
+	 * @param 	InHLODLayer		The HLOD layer for which we'll create a setting object
+	 * @return A newly created UHLODBuilderSettings object, outered to the provided HLOD layer.
+	 */
+	virtual UHLODBuilderSettings* CreateHLODBuilderSettings(UHLODLayer* InHLODLayer) = 0;
 };
diff --git a/Engine/Source/Developer/WorldPartitionHLODUtilities/Public/WorldPartitionHLODUtilities.h b/Engine/Source/Developer/WorldPartitionHLODUtilities/Public/WorldPartitionHLODUtilities.h
--- a/Engine/Source/Developer/WorldPartitionHLODUtilities/Public/WorldPartitionHLODUtilities.h
+++ b/Engine/Source/Developer/WorldPartitionHLODUtilities/Public/WorldPartitionHLODUtilities.h
@@ -12,4 +12,5 @@ class FWorldPartitionHLODUtilities : public IWorldPartitionHLODUtilities
 public:
 	virtual TArray<AWorldPartitionHLOD*> CreateHLODActors(FHLODCreationContext& InCreationContext, const FHLODCreationParams& InCreationParams, const TSet<FActorInstance>& InActors, const TArray<const UDataLayer*>& InDataLayers) override;
 	virtual uint32 BuildHLOD(AWorldPartitionHLOD* InHLODActor) override;
+	virtual UHLODBuilderSettings* CreateHLODBuilderSettings(UHLODLayer* InHLODLayer) override;
 };
diff --git a/Engine/Source/Runtime/Engine/Private/WorldPartition/HLOD/HLODBuilder.cpp b/Engine/Source/Runtime/Engine/Private/WorldPartition/HLOD/HLODBuilder.cpp
new file mode 100644
--- /dev/null
+++ b/Engine/Source/Runtime/Engine/Private/WorldPartition/HLOD/HLODBuilder.cpp
@@ -0,0 +1,145 @@
+// Copyright Epic Games, Inc. All Rights Reserved.
+
+#include "WorldPartition/HLOD/HLODBuilder.h"
+
+#include "Engine/StaticMesh.h"
+#include "Components/InstancedStaticMeshComponent.h"
+
+#include "WorldPartition/HLOD/HLODActor.h"
+#include "WorldPartition/HLOD/HLODLayer.h"
+
+#include "AssetCompilingManager.h"
+
+DEFINE_LOG_CATEGORY(LogHLODBuilder);
+
+
+UHLODBuilder::UHLODBuilder(const FObjectInitializer& ObjectInitializer)
+	: Super(ObjectInitializer)
+{
+}
+
+UHLODBuilderSettings::UHLODBuilderSettings(const FObjectInitializer& ObjectInitializer)
+	: Super(ObjectInitializer)
+{
+}
+
+#if WITH_EDITOR
+
+UHLODBuilderSettings* UHLODBuilder::CreateSettings(UHLODLayer* InHLODLayer) const
+{
+	return NewObject<UHLODBuilderSettings>(InHLODLayer);
+}
+
+bool UHLODBuilder::RequiresCompiledAssets() const
+{
+	return true;
+}
+
+TArray<UPrimitiveComponent*> UHLODBuilder::CreateComponents(AWorldPartitionHLOD* InHLODActor, const UHLODLayer* InHLODLayer, const TArray<UPrimitiveComponent*>& InSubComponents) const
+{
+	return TArray<UPrimitiveComponent*>();
+}
+
+void UHLODBuilder::Build(AWorldPartitionHLOD* InHLODActor, const UHLODLayer* InHLODLayer, const TArray<FWorldPartitionReference>& InSubActors)
+{
+	TArray<UPrimitiveComponent*> SubComponents = GatherPrimitiveComponents(InSubActors);
+	if (SubComponents.IsEmpty())
+	{
+		return;
+	}
+
+	// Before capturing the scene, make sure all assets are finished compiling
+	FAssetCompilingManager::Get().FinishAllCompilation();
+
+	TArray<UPrimitiveComponent*> HLODPrimitives = CreateComponents(InHLODActor, InHLODLayer, SubComponents);
+	HLODPrimitives.RemoveSwap(nullptr);
+
+	for (UPrimitiveComponent* HLODPrimitive : HLODPrimitives)
+	{
+		DisableCollisions(HLODPrimitive);
+
+		HLODPrimitive->bComputeFastLocalBounds = true;
+		HLODPrimitive->bComputeBoundsOnceDuringCook = true;
+	}
+
+	if (!HLODPrimitives.IsEmpty())
+	{
+		InHLODActor->Modify();
+		InHLODActor->SetHLODPrimitives(HLODPrimitives);
+	}
+}
+
+TArray<UPrimitiveComponent*> UHLODBuilder::GatherPrimitiveComponents(const TArray<FWorldPartitionReference>& InActors)
+{
+	TArray<UPrimitiveComponent*> PrimitiveComponents;
+
+	auto GatherPrimitivesFromActor = [&PrimitiveComponents](const AActor* Actor, const AActor* ParentActor = nullptr)
+	{
+		const TCHAR* Padding = ParentActor ? TEXT("    ") : TEXT("");
+		UE_LOG(LogHLODBuilder, Verbose, TEXT("%s* Adding components from actor %s"), Padding, *Actor->GetName());
+		for (UActorComponent* SubComponent : Actor->GetComponents())
+		{
+			if (SubComponent && SubComponent->IsHLODRelevant())
+			{
+				if (UPrimitiveComponent* PrimitiveComponent = Cast<UPrimitiveComponent>(SubComponent))
+				{
+					PrimitiveComponents.Add(PrimitiveComponent);
+
+					if (UInstancedStaticMeshComponent* ISMC = Cast<UInstancedStaticMeshComponent>(PrimitiveComponent))
+					{
+						UE_LOG(LogHLODBuilder, Verbose, TEXT("%s    * %s [%d instances]"), Padding, *ISMC->GetStaticMesh()->GetName(), ISMC->GetInstanceCount());
+					}
+					else if (UStaticMeshComponent* SMC = Cast<UStaticMeshComponent>(PrimitiveComponent))
+					{
+						UE_LOG(LogHLODBuilder, Verbose, TEXT("%s    * %s"), Padding, *SMC->GetStaticMesh()->GetName());
+					}
+				}
+				else
+				{
+					UE_LOG(LogHLODBuilder, Warning, TEXT("Component \"%s\" is marked as HLOD-relevant but this type of component currently unsupported."), *SubComponent->GetFullName());
+				}
+			}
+		}
+	};
+
+	TSet<AActor*> UnderlyingActors;
+
+	for (const FWorldPartitionReference& ActorRef : InActors)
+	{
+		AActor* Actor = ActorRef->GetActor();
+		if (!Actor->IsHLODRelevant())
+		{
+			continue;
+		}
+
+		// Gather primitives from the Actor
+		GatherPrimitivesFromActor(Actor);
+
+		// Retrieve all underlying actors (ex: all sub actors of a LevelInstance)
+		UnderlyingActors.Reset();
+		Actor->EditorGetUnderlyingActors(UnderlyingActors);
+
+		// Gather primitives from underlying actors
+		for (const AActor* UnderlyingActor : UnderlyingActors)
+		{
+			if (UnderlyingActor->IsHLODRelevant())
+			{
+				GatherPrimitivesFromActor(UnderlyingActor, Actor);
+			}
+		}
+	}
+
+	return PrimitiveComponents;
+}
+
+void UHLODBuilder::DisableCollisions(UPrimitiveComponent* Component)
+{
+	Component->SetCollisionProfileName(UCollisionProfile::NoCollision_ProfileName);
+	Component->SetGenerateOverlapEvents(false);
+	Component->SetCanEverAffectNavigation(false);
+	Component->CanCharacterStepUpOn = ECanBeCharacterBase::ECB_No;
+	Component->SetCanEverAffectNavigation(false);
+	Component->SetCollisionEnabled(ECollisionEnabled::NoCollision);
+}
+
+#endif // WITH_EDITOR
diff --git a/Engine/Source/Runtime/Engine/Private/WorldPartition/HLOD/HLODLayer.cpp b/Engine/Source/Runtime/Engine/Private/WorldPartition/HLOD/HLODLayer.cpp
--- a/Engine/Source/Runtime/Engine/Private/WorldPartition/HLOD/HLODLayer.cpp
+++ b/Engine/Source/Runtime/Engine/Private/WorldPartition/HLOD/HLODLayer.cpp
@@ -8,15 +8,17 @@
 #include "WorldPartition/HLOD/HLODActor.h"
 
 #if WITH_EDITOR
-#include "Algo/Copy.h"
-#include "UObject/ConstructorHelpers.h"
-#include "Materials/Material.h"
-#include "Engine/HLODProxy.h"
 #include "Serialization/ArchiveCrc32.h"
 
+#include "Engine/World.h"
+
 #include "WorldPartition/WorldPartition.h"
 #include "WorldPartition/WorldPartitionActorDesc.h"
 #include "WorldPartition/WorldPartitionActorDescView.h"
+
+#include "Modules/ModuleManager.h"
+#include "IWorldPartitionHLODUtilities.h"
+#include "WorldPartitionHLODUtilitiesModule.h"
 #endif
 
 DEFINE_LOG_CATEGORY_STATIC(LogHLODLayer, Log, All);
@@ -28,61 +30,10 @@ UHLODLayer::UHLODLayer(const FObjectInitializer& ObjectInitializer)
 	, LoadingRange(12800)
 #endif
 {
-#if WITH_EDITORONLY_DATA
-	if (!IsTemplate())
-	{
-		HLODMaterial = GEngine->DefaultHLODFlattenMaterial;
-	}
-#endif
 }
 
 #if WITH_EDITOR
 
-uint32 UHLODLayer::GetCRC() const
-{
-	UHLODLayer& This = *const_cast<UHLODLayer*>(this);
-
-	FArchiveCrc32 Ar;
-
-	Ar << This.LayerType;
-	UE_LOG(LogHLODLayer, VeryVerbose, TEXT(" - LayerType = %d"), Ar.GetCrc());
-
-	if (LayerType == EHLODLayerType::MeshMerge)
-	{
-		Ar << This.MeshMergeSettings;
-		UE_LOG(LogHLODLayer, VeryVerbose, TEXT(" - MeshMergeSettings = %d"), Ar.GetCrc());
-	}
-	else if (LayerType == EHLODLayerType::MeshSimplify)
-	{
-		Ar << This.MeshSimplifySettings;
-		UE_LOG(LogHLODLayer, VeryVerbose, TEXT(" - MeshSimplifySettings = %d"), Ar.GetCrc());
-	}
-	else if (LayerType == EHLODLayerType::MeshApproximate)
-	{
-		Ar << This.MeshApproximationSettings;
-		UE_LOG(LogHLODLayer, VeryVerbose, TEXT(" - MeshApproximationSettings = %d"), Ar.GetCrc());
-	}
-
-	Ar << This.CellSize;
-	UE_LOG(LogHLODLayer, VeryVerbose, TEXT(" - CellSize = %d"), Ar.GetCrc());
-
-	uint32 Hash = Ar.GetCrc();
-
-	const bool bUseHLODMaterial = LayerType != EHLODLayerType::Instancing;
-	if (bUseHLODMaterial && !HLODMaterial.IsNull())
-	{
-		UMaterialInterface* Material = HLODMaterial.LoadSynchronous();
-		if (Material)
-		{
-			uint32 MaterialCRC = UHLODProxy::GetCRC(Material);
-			UE_LOG(LogHLODLayer, VeryVerbose, TEXT(" - Material = %d"), MaterialCRC);
-			Hash = HashCombine(Hash, MaterialCRC);
-		}
-	}
-
-	return Hash;
-}
-
 UHLODLayer* UHLODLayer::GetHLODLayer(const AActor* InActor)
 {
 	if (UHLODLayer* HLODLayer = InActor->GetHLODLayer())
@@ -129,6 +80,35 @@ UHLODLayer* UHLODLayer::GetHLODLayer(const FWorldPartitionActorDesc& InActorDesc
 	return GetHLODLayer(FWorldPartitionActorDescView(&InActorDesc), InWorldPartition);
 }
 
+void UHLODLayer::PostLoad()
+{
+	Super::PostLoad();
+
+	if (HLODBuilderSettings == nullptr)
+	{
+		IWorldPartitionHLODUtilities* WPHLODUtilities = FModuleManager::Get().LoadModuleChecked<IWorldPartitionHLODUtilitiesModule>("WorldPartitionHLODUtilities").GetUtilities();
+		if (WPHLODUtilities)
+		{
+			HLODBuilderSettings = WPHLODUtilities->CreateHLODBuilderSettings(this);
+		}
+	}
+}
+
+void UHLODLayer::PostEditChangeProperty(FPropertyChangedEvent& PropertyChangedEvent)
+{
+	Super::PostEditChangeProperty(PropertyChangedEvent);
+
+	const FName PropertyName = PropertyChangedEvent.Property ? PropertyChangedEvent.Property->GetFName() : FName();
+
+	IWorldPartitionHLODUtilities* WPHLODUtilities = FModuleManager::Get().LoadModuleChecked<IWorldPartitionHLODUtilitiesModule>("WorldPartitionHLODUtilities").GetUtilities();
+
+	if (PropertyName == GET_MEMBER_NAME_CHECKED(UHLODLayer, LayerType) ||
+		PropertyName == GET_MEMBER_NAME_CHECKED(UHLODLayer, HLODBuilderClass))
+	{
+		HLODBuilderSettings = WPHLODUtilities->CreateHLODBuilderSettings(this);
+	}
+}
+
 #endif // WITH_EDITOR
 
 #if WITH_EDITORONLY_DATA
diff --git a/Engine/Source/Runtime/Engine/Public/WorldPartition/HLOD/HLODBuilder.h b/Engine/Source/Runtime/Engine/Public/WorldPartition/HLOD/HLODBuilder.h
new file mode 100644
--- /dev/null
+++ b/Engine/Source/Runtime/Engine/Public/WorldPartition/HLOD/HLODBuilder.h
@@ -0,0 +1,49 @@
+// Copyright Epic Games, Inc. All Rights Reserved.
+
+#pragma once
+
+#include "CoreMinimal.h"
+#include "WorldPartition/WorldPartitionHandle.h"
+#include "HLODBuilder.generated.h"
+
+class AWorldPartitionHLOD;
+class UPrimitiveComponent;
+class UHLODLayer;
+
+ENGINE_API DECLARE_LOG_CATEGORY_EXTERN(LogHLODBuilder, Log, All);
+
+
+// This is the base class for all HLOD Builder settings
+UCLASS()
+class ENGINE_API UHLODBuilderSettings : public UObject
+{
+	GENERATED_UCLASS_BODY()
+#if WITH_EDITOR
+	virtual uint32 GetCRC() const { return 0; }
+#endif // WITH_EDITOR
+};
+
+
+/**
+ * Base class for all HLODBuilders
+ */
+UCLASS(Abstract)
+class ENGINE_API UHLODBuilder : public UObject
+{
+	 GENERATED_UCLASS_BODY()
+
+#if WITH_EDITOR
+public:
+	virtual UHLODBuilderSettings* CreateSettings(UHLODLayer* InHLODLayer) const;
+
+	virtual bool RequiresCompiledAssets() const;
+
+	virtual TArray<UPrimitiveComponent*> CreateComponents(AWorldPartitionHLOD* InHLODActor, const UHLODLayer* InHLODLayer, const TArray<UPrimitiveComponent*>& InSubComponents) const;
+
+	void Build(AWorldPartitionHLOD* InHLODActor, const UHLODLayer* InHLODLayer, const TArray<FWorldPartitionReference>& InSubActors);
+
+	static TArray<UPrimitiveComponent*> GatherPrimitiveComponents(const TArray<FWorldPartitionReference>& InActors);
+
+	static void DisableCollisions(UPrimitiveComponent* Component);
+#endif // WITH_EDITOR
+};
diff --git a/Engine/Source/Runtime/Engine/Public/WorldPartition/HLOD/HLODLayer.h b/Engine/Source/Runtime/Engine/Public/WorldPartition/HLOD/HLODLayer.h
--- a/Engine/Source/Runtime/Engine/Public/WorldPartition/HLOD/HLODLayer.h
+++ b/Engine/Source/Runtime/Engine/Public/WorldPartition/HLOD/HLODLayer.h
@@ -9,10 +9,9 @@
 #include "UObject/Object.h"
 
 #include "EngineDefines.h"
-
-#if WITH_EDITOR
 #include "Engine/MeshMerging.h"
-#endif
+
+#include "WorldPartition/HLOD/HLODBuilder.h"
 
 #include "HLODLayer.generated.h"
 
@@ -25,10 +24,11 @@ class FWorldPartitionActorDescView;
 UENUM()
 enum class EHLODLayerType : uint8
 {
-	Instancing = 0			UMETA(DisplayName = "Instancing"),
-	MeshMerge = 1			UMETA(DisplayName = "Merged Mesh"),
-	MeshSimplify = 2		UMETA(DisplayName = "Simplified Mesh"),
-	MeshApproximate = 3		UMETA(DisplayName = "Approximated Mesh"),
+	Instancing				UMETA(DisplayName = "Instancing"),
+	MeshMerge				UMETA(DisplayName = "Merged Mesh"),
+	MeshSimplify			UMETA(DisplayName = "Simplified Mesh"),
+	MeshApproximate			UMETA(DisplayName = "Approximated Mesh"),
+	Custom					UMETA(DisplayName = "Custom"),
 };
 
 UCLASS(Blueprintable, Config=Engine, PerObjectConfig)
@@ -41,16 +41,12 @@ public:
 	static UHLODLayer* GetHLODLayer(const AActor* InActor);
 	static UHLODLayer* GetHLODLayer(const FWorldPartitionActorDesc& InActorDesc, const UWorldPartition* InWorldPartition);
 	static UHLODLayer* GetHLODLayer(const FWorldPartitionActorDescView& InActorDescView, const UWorldPartition* InWorldPartition);
-
-	uint32 GetCRC() const;
 #endif
 
 #if WITH_EDITORONLY_DATA
 	EHLODLayerType GetLayerType() const { return LayerType; }
-	const FMeshMergingSettings& GetMeshMergeSettings() const { return MeshMergeSettings; }
-	const FMeshProxySettings& GetMeshSimplifySettings() const { return MeshSimplifySettings; }
-	const FMeshApproximationSettings& GetMeshApproximationSettings() const { return MeshApproximationSettings; }
-	const TSoftObjectPtr<UMaterial>& GetHLODMaterial() const { return HLODMaterial; }
+	const TSubclassOf<UHLODBuilder> GetHLODBuilderClass() const { return HLODBuilderClass; }
+	const UHLODBuilderSettings* GetHLODBuilderSettings() const { return HLODBuilderSettings; }
 	FName GetRuntimeGrid(uint32 InHLODLevel) const;
 	int32 GetCellSize() const { return bAlwaysLoaded ? 0 : CellSize; }
 	float GetLoadingRange() const { return bAlwaysLoaded ? WORLD_MAX : LoadingRange; }
@@ -60,27 +56,26 @@ public:
 	static FName GetRuntimeGridName(uint32 InLODLevel, int32 InCellSize, float InLoadingRange);
 #endif
 
+#if WITH_EDITOR
+private:
+	//~ Begin UObject Interface.
+	virtual void PostLoad() override;
+	virtual void PostEditChangeProperty(FPropertyChangedEvent& PropertyChangedEvent) override;
+	//~ End UObject Interface.
+#endif
+
 #if WITH_EDITORONLY_DATA
 private:
 	/** Type of HLOD generation to use */
 	UPROPERTY(EditAnywhere, Config, Category=HLOD)
 	EHLODLayerType LayerType;
 
-	/** Merged mesh generation settings - Used when this layer is of type EHLODLayerType::MeshMerge */
-	UPROPERTY(EditAnywhere, Config, Category=HLOD, meta = (EditConditionHides, EditCondition = "LayerType == EHLODLayerType::MeshMerge"))
-	FMeshMergingSettings MeshMergeSettings;
-
-	/** Simplified mesh generation settings - Used when this layer is of type EHLODLayerType::MeshSimplify */
-	UPROPERTY(EditAnywhere, Config, Category=HLOD, meta = (EditConditionHides, EditCondition = "LayerType == EHLODLayerType::MeshSimplify"))
-	FMeshProxySettings MeshSimplifySettings;
+	/** HLODBuilder class */
+	UPROPERTY(EditAnywhere, Config, Category=HLOD, meta = (EditConditionHides, EditCondition = "LayerType == EHLODLayerType::Custom"))
+	TSubclassOf<UHLODBuilder> HLODBuilderClass;
 
-	/** Mesh approximation settings - Used when this layer is of type EHLODLayerType::MeshApproximation */
-	UPROPERTY(EditAnywhere, Config, Category = HLOD, meta = (EditConditionHides, EditCondition = "LayerType == EHLODLayerType::MeshApproximate"))
-	FMeshApproximationSettings MeshApproximationSettings;
-
-	/** Material that will be used by the generated HLOD static mesh */
-	UPROPERTY(EditAnywhere, Config, AdvancedDisplay, Category=HLOD, meta = (EditConditionHides, EditCondition = "LayerType == EHLODLayerType::MeshMerge || LayerType == EHLODLayerType::MeshSimplify"))
-	TSoftObjectPtr<UMaterial> HLODMaterial;
+	UPROPERTY(VisibleAnywhere, Export, NoClear, Category=HLOD, meta = (EditInline, NoResetToDefault))
+	UHLODBuilderSettings* HLODBuilderSettings;
 
 	/** Whether HLOD actors generated for this layer will be always loaded */
 	UPROPERTY(EditAnywhere, Config, Category=HLOD)
@@ -97,5 +92,15 @@ private:
 	/** HLOD Layer to assign to the generated HLOD actors */
 	UPROPERTY(EditAnywhere, Config, Category=HLOD, meta = (EditConditionHides, EditCondition = "!bAlwaysLoaded"))
 	TSoftObjectPtr<UHLODLayer> ParentLayer;
+
+public:
+	UPROPERTY()
+	FMeshMergingSettings MeshMergeSettings_DEPRECATED;
+	UPROPERTY()
+	FMeshProxySettings MeshSimplifySettings_DEPRECATED;
+	UPROPERTY()
+	FMeshApproximationSettings MeshApproximationSettings_DEPRECATED;
+	UPROPERTY()
+	TSoftObjectPtr<UMaterial> HLODMaterial_DEPRECATED;
 #endif
 };
