HLOD builders - Expose base builder classes outside of the WorldPartitionHLODUtilities module
* Corrected wrong profiling tags
* Removed duplicate call to FAssetCompilingManager::Get().FinishAllCompilation() (already performed in  FWorldPartitionHLODUtilities::BuildHLOD())

diff --git a/Engine/Source/Developer/WorldPartitionHLODUtilities/Private/HLODBuilderInstancing.cpp b/Engine/Source/Developer/WorldPartitionHLODUtilities/Private/HLODBuilderInstancing.cpp
--- a/Engine/Source/Developer/WorldPartitionHLODUtilities/Private/HLODBuilderInstancing.cpp
+++ b/Engine/Source/Developer/WorldPartitionHLODUtilities/Private/HLODBuilderInstancing.cpp
@@ -45,7 +45,7 @@ UHLODBuilderInstancing::UHLODBuilderInstancing(const FObjectInitializer& ObjectI
 
 TArray<UPrimitiveComponent*> UHLODBuilderInstancing::CreateComponents(AWorldPartitionHLOD* InHLODActor, const UHLODLayer* InHLODLayer, const TArray<UPrimitiveComponent*>& InSubComponents) const
 {
-	TRACE_CPUPROFILER_EVENT_SCOPE(FHLODBuilder_Instancing::CreateComponents);
+	TRACE_CPUPROFILER_EVENT_SCOPE(UHLODBuilderInstancing::CreateComponents);
 
 	TArray<UPrimitiveComponent*> Components;
 
diff --git a/Engine/Source/Developer/WorldPartitionHLODUtilities/Private/HLODBuilderInstancing.h b/Engine/Source/Developer/WorldPartitionHLODUtilities/Private/HLODBuilderInstancing.h
deleted file mode 100644
diff --git a/Engine/Source/Developer/WorldPartitionHLODUtilities/Private/HLODBuilderMeshApproximate.cpp b/Engine/Source/Developer/WorldPartitionHLODUtilities/Private/HLODBuilderMeshApproximate.cpp
--- a/Engine/Source/Developer/WorldPartitionHLODUtilities/Private/HLODBuilderMeshApproximate.cpp
+++ b/Engine/Source/Developer/WorldPartitionHLODUtilities/Private/HLODBuilderMeshApproximate.cpp
@@ -80,7 +80,7 @@ UHLODBuilderSettings* UHLODBuilderMeshApproximate::CreateSettings(UHLODLayer* In
 
 TArray<UPrimitiveComponent*> UHLODBuilderMeshApproximate::CreateComponents(AWorldPartitionHLOD* InHLODActor, const UHLODLayer* InHLODLayer, const TArray<UPrimitiveComponent*>& InSubComponents) const
 {
-	TRACE_CPUPROFILER_EVENT_SCOPE(FHLODBuilder_MeshApproximate::CreateComponents);
+	TRACE_CPUPROFILER_EVENT_SCOPE(UHLODBuilderMeshApproximate::CreateComponents);
 
 	TSet<AActor*> Actors;
 	TArray<UPrimitiveComponent*> InstancedComponents;
diff --git a/Engine/Source/Developer/WorldPartitionHLODUtilities/Private/HLODBuilderMeshApproximate.h b/Engine/Source/Developer/WorldPartitionHLODUtilities/Private/HLODBuilderMeshApproximate.h
deleted file mode 100644
diff --git a/Engine/Source/Developer/WorldPartitionHLODUtilities/Private/HLODBuilderMeshMerge.cpp b/Engine/Source/Developer/WorldPartitionHLODUtilities/Private/HLODBuilderMeshMerge.cpp
--- a/Engine/Source/Developer/WorldPartitionHLODUtilities/Private/HLODBuilderMeshMerge.cpp
+++ b/Engine/Source/Developer/WorldPartitionHLODUtilities/Private/HLODBuilderMeshMerge.cpp
@@ -76,7 +76,7 @@ UHLODBuilderSettings* UHLODBuilderMeshMerge::CreateSettings(UHLODLayer* InHLODLa
 
 TArray<UPrimitiveComponent*> UHLODBuilderMeshMerge::CreateComponents(AWorldPartitionHLOD* InHLODActor, const UHLODLayer* InHLODLayer, const TArray<UPrimitiveComponent*>& InSubComponents) const
 {
-	TRACE_CPUPROFILER_EVENT_SCOPE(FHLODBuilder_MeshMerge::CreateComponents);
+	TRACE_CPUPROFILER_EVENT_SCOPE(UHLODBuilderMeshMerge::CreateComponents);
 
 	TArray<UObject*> Assets;
 	FVector MergedActorLocation;
diff --git a/Engine/Source/Developer/WorldPartitionHLODUtilities/Private/HLODBuilderMeshMerge.h b/Engine/Source/Developer/WorldPartitionHLODUtilities/Private/HLODBuilderMeshMerge.h
deleted file mode 100644
diff --git a/Engine/Source/Developer/WorldPartitionHLODUtilities/Private/HLODBuilderMeshSimplify.cpp b/Engine/Source/Developer/WorldPartitionHLODUtilities/Private/HLODBuilderMeshSimplify.cpp
--- a/Engine/Source/Developer/WorldPartitionHLODUtilities/Private/HLODBuilderMeshSimplify.cpp
+++ b/Engine/Source/Developer/WorldPartitionHLODUtilities/Private/HLODBuilderMeshSimplify.cpp
@@ -79,7 +79,7 @@ UHLODBuilderSettings* UHLODBuilderMeshSimplify::CreateSettings(UHLODLayer* InHLO
 
 TArray<UPrimitiveComponent*> UHLODBuilderMeshSimplify::CreateComponents(AWorldPartitionHLOD* InHLODActor, const UHLODLayer* InHLODLayer, const TArray<UPrimitiveComponent*>& InSubComponents) const
 {
-	TRACE_CPUPROFILER_EVENT_SCOPE(FHLODBuilder_MeshSimplify::CreateComponents);
+	TRACE_CPUPROFILER_EVENT_SCOPE(UHLODBuilderMeshSimplifySettings::CreateComponents);
 
 	TArray<UObject*> Assets;
 	FCreateProxyDelegate ProxyDelegate;
diff --git a/Engine/Source/Developer/WorldPartitionHLODUtilities/Private/HLODBuilderMeshSimplify.h b/Engine/Source/Developer/WorldPartitionHLODUtilities/Private/HLODBuilderMeshSimplify.h
deleted file mode 100644
diff --git a/Engine/Source/Developer/WorldPartitionHLODUtilities/Public/HLODBuilderInstancing.h b/Engine/Source/Developer/WorldPartitionHLODUtilities/Public/HLODBuilderInstancing.h
new file mode 100644
--- /dev/null
+++ b/Engine/Source/Developer/WorldPartitionHLODUtilities/Public/HLODBuilderInstancing.h
@@ -0,0 +1,21 @@
+// Copyright Epic Games, Inc. All Rights Reserved.
+
+#pragma once
+
+#include "WorldPartition/HLOD/HLODBuilder.h"
+#include "HLODBuilderInstancing.generated.h"
+
+
+/**
+ * Build a AWorldPartitionHLOD whose components are ISMC
+ */
+UCLASS()
+class WORLDPARTITIONHLODUTILITIES_API UHLODBuilderInstancing : public UHLODBuilder
+{
+	 GENERATED_UCLASS_BODY()
+
+public:
+	virtual bool RequiresCompiledAssets() const override { return false; }
+
+	virtual TArray<UPrimitiveComponent*> CreateComponents(AWorldPartitionHLOD* InHLODActor, const UHLODLayer* InHLODLayer, const TArray<UPrimitiveComponent*>& InSubComponents) const override;
+};
diff --git a/Engine/Source/Developer/WorldPartitionHLODUtilities/Public/HLODBuilderMeshApproximate.h b/Engine/Source/Developer/WorldPartitionHLODUtilities/Public/HLODBuilderMeshApproximate.h
new file mode 100644
--- /dev/null
+++ b/Engine/Source/Developer/WorldPartitionHLODUtilities/Public/HLODBuilderMeshApproximate.h
@@ -0,0 +1,41 @@
+// Copyright Epic Games, Inc. All Rights Reserved.
+
+#pragma once
+
+#include "WorldPartition/HLOD/HLODBuilder.h"
+#include "Engine/MeshMerging.h"
+
+#include "HLODBuilderMeshApproximate.generated.h"
+
+class UMaterial;
+
+
+UCLASS(Blueprintable, Config = Engine, PerObjectConfig)
+class WORLDPARTITIONHLODUTILITIES_API UHLODBuilderMeshApproximateSettings : public UHLODBuilderSettings
+{
+	GENERATED_UCLASS_BODY()
+
+	virtual uint32 GetCRC() const override;
+
+	/** Mesh approximation settings */
+	UPROPERTY(EditAnywhere, Config, Category = HLOD, meta=(EditInline))
+	FMeshApproximationSettings MeshApproximationSettings;
+
+	/** Material that will be used by the generated HLOD static mesh */
+	UPROPERTY(EditAnywhere, Config, AdvancedDisplay, Category = HLOD)
+	TSoftObjectPtr<UMaterial> HLODMaterial;
+};
+
+
+/**
+ * Build an approximated mesh using geometry from the provided actors
+ */
+UCLASS()
+class WORLDPARTITIONHLODUTILITIES_API UHLODBuilderMeshApproximate : public UHLODBuilder
+{
+	GENERATED_UCLASS_BODY()
+
+public:
+	virtual UHLODBuilderSettings* CreateSettings(UHLODLayer* InHLODLayer) const override;
+	virtual TArray<UPrimitiveComponent*> CreateComponents(AWorldPartitionHLOD* InHLODActor, const UHLODLayer* InHLODLayer, const TArray<UPrimitiveComponent*>& InSubComponents) const override;
+};
diff --git a/Engine/Source/Developer/WorldPartitionHLODUtilities/Public/HLODBuilderMeshMerge.h b/Engine/Source/Developer/WorldPartitionHLODUtilities/Public/HLODBuilderMeshMerge.h
new file mode 100644
--- /dev/null
+++ b/Engine/Source/Developer/WorldPartitionHLODUtilities/Public/HLODBuilderMeshMerge.h
@@ -0,0 +1,41 @@
+// Copyright Epic Games, Inc. All Rights Reserved.
+
+#pragma once
+
+#include "WorldPartition/HLOD/HLODBuilder.h"
+#include "Engine/MeshMerging.h"
+
+#include "HLODBuilderMeshMerge.generated.h"
+
+class UMaterial;
+
+
+UCLASS(Blueprintable, Config = Engine, PerObjectConfig)
+class WORLDPARTITIONHLODUTILITIES_API UHLODBuilderMeshMergeSettings : public UHLODBuilderSettings
+{
+	GENERATED_UCLASS_BODY()
+
+	virtual uint32 GetCRC() const override;
+
+	/** Merged mesh generation settings */
+	UPROPERTY(EditAnywhere, Config, Category = HLOD)
+	FMeshMergingSettings MeshMergeSettings;
+
+	/** Material that will be used by the generated HLOD static mesh */
+	UPROPERTY(EditAnywhere, Config, AdvancedDisplay, Category = HLOD)
+	TSoftObjectPtr<UMaterial> HLODMaterial;
+};
+
+
+/**
+ * Build a merged mesh using geometry from the provided actors
+ */
+UCLASS()
+class WORLDPARTITIONHLODUTILITIES_API UHLODBuilderMeshMerge : public UHLODBuilder
+{
+	GENERATED_UCLASS_BODY()
+
+public:
+	virtual UHLODBuilderSettings* CreateSettings(UHLODLayer* InHLODLayer) const override;
+	virtual TArray<UPrimitiveComponent*> CreateComponents(AWorldPartitionHLOD* InHLODActor, const UHLODLayer* InHLODLayer, const TArray<UPrimitiveComponent*>& InSubComponents) const override;
+};
diff --git a/Engine/Source/Developer/WorldPartitionHLODUtilities/Public/HLODBuilderMeshSimplify.h b/Engine/Source/Developer/WorldPartitionHLODUtilities/Public/HLODBuilderMeshSimplify.h
new file mode 100644
--- /dev/null
+++ b/Engine/Source/Developer/WorldPartitionHLODUtilities/Public/HLODBuilderMeshSimplify.h
@@ -0,0 +1,41 @@
+// Copyright Epic Games, Inc. All Rights Reserved.
+
+#pragma once
+
+#include "WorldPartition/HLOD/HLODBuilder.h"
+#include "Engine/MeshMerging.h"
+
+#include "HLODBuilderMeshSimplify.generated.h"
+
+class UMaterial;
+
+
+UCLASS(Blueprintable, Config = Engine, PerObjectConfig)
+class WORLDPARTITIONHLODUTILITIES_API UHLODBuilderMeshSimplifySettings : public UHLODBuilderSettings
+{
+	GENERATED_UCLASS_BODY()
+
+	virtual uint32 GetCRC() const override;
+
+	/** Simplified mesh generation settings */
+	UPROPERTY(EditAnywhere, Config, Category = HLOD)
+	FMeshProxySettings MeshSimplifySettings;
+
+	/** Material that will be used by the generated HLOD static mesh */
+	UPROPERTY(EditAnywhere, Config, AdvancedDisplay, Category = HLOD)
+	TSoftObjectPtr<UMaterial> HLODMaterial;
+};
+
+
+/**
+ * Build a simplified mesh using geometry from the provided actors
+ */
+UCLASS()
+class WORLDPARTITIONHLODUTILITIES_API UHLODBuilderMeshSimplify : public UHLODBuilder
+{
+	GENERATED_UCLASS_BODY()
+
+public:
+	virtual UHLODBuilderSettings* CreateSettings(UHLODLayer* InHLODLayer) const override;
+	virtual TArray<UPrimitiveComponent*> CreateComponents(AWorldPartitionHLOD* InHLODActor, const UHLODLayer* InHLODLayer, const TArray<UPrimitiveComponent*>& InSubComponents) const override;
+};
diff --git a/Engine/Source/Runtime/Engine/Private/WorldPartition/HLOD/HLODBuilder.cpp b/Engine/Source/Runtime/Engine/Private/WorldPartition/HLOD/HLODBuilder.cpp
--- a/Engine/Source/Runtime/Engine/Private/WorldPartition/HLOD/HLODBuilder.cpp
+++ b/Engine/Source/Runtime/Engine/Private/WorldPartition/HLOD/HLODBuilder.cpp
@@ -48,9 +48,6 @@ void UHLODBuilder::Build(AWorldPartitionHLOD* InHLODActor, const UHLODLayer* InH
 		return;
 	}
 
-	// Before capturing the scene, make sure all assets are finished compiling
-	FAssetCompilingManager::Get().FinishAllCompilation();
-
 	TArray<UPrimitiveComponent*> HLODPrimitives = CreateComponents(InHLODActor, InHLODLayer, SubComponents);
 	HLODPrimitives.RemoveSwap(nullptr);
 
