Fix HLOD resources warmup being performed for the transition to HLOD0. It should only be used for HLOD0 -> HLOD1

diff --git a/Engine/Source/Developer/WorldPartitionHLODUtilities/Private/WorldPartitionHLODUtilities.cpp b/Engine/Source/Developer/WorldPartitionHLODUtilities/Private/WorldPartitionHLODUtilities.cpp
--- a/Engine/Source/Developer/WorldPartitionHLODUtilities/Private/WorldPartitionHLODUtilities.cpp
+++ b/Engine/Source/Developer/WorldPartitionHLODUtilities/Private/WorldPartitionHLODUtilities.cpp
@@ -315,7 +315,7 @@ TArray<AWorldPartitionHLOD*> FWorldPartitionHLODUtilities::CreateHLODActors(FHLO
 	return HLODActors;
 }
 
-static TSubclassOf<UHLODBuilder> GetHLODBuilderClassForLayer(const UHLODLayer* InHLODLayer)
+TSubclassOf<UHLODBuilder> FWorldPartitionHLODUtilities::GetHLODBuilderClass(const UHLODLayer* InHLODLayer)
 {
 	EHLODLayerType HLODLayerType = InHLODLayer->GetLayerType();
 	switch (HLODLayerType)
@@ -348,7 +348,7 @@ static TSubclassOf<UHLODBuilder> GetHLODBuilderClassForLayer(const UHLODLayer* I
 
 UHLODBuilderSettings* FWorldPartitionHLODUtilities::CreateHLODBuilderSettings(UHLODLayer* InHLODLayer)
 {
-	TSubclassOf<UHLODBuilder> HLODBuilderClass = GetHLODBuilderClassForLayer(InHLODLayer);
+	TSubclassOf<UHLODBuilder> HLODBuilderClass = GetHLODBuilderClass(InHLODLayer);
 	if (!HLODBuilderClass)
 	{
 		return NewObject<UHLODBuilderSettings>(InHLODLayer, UHLODBuilderSettings::StaticClass());
@@ -382,7 +382,7 @@ uint32 FWorldPartitionHLODUtilities::BuildHLOD(AWorldPartitionHLOD* InHLODActor)
 	}
 
 	const UHLODLayer* HLODLayer = InHLODActor->GetSubActorsHLODLayer();
-	TSubclassOf<UHLODBuilder> HLODBuilderClass = GetHLODBuilderClassForLayer(HLODLayer);
+	TSubclassOf<UHLODBuilder> HLODBuilderClass = GetHLODBuilderClass(HLODLayer);
 
 	if (HLODBuilderClass)
 	{
diff --git a/Engine/Source/Developer/WorldPartitionHLODUtilities/Public/IWorldPartitionHLODUtilities.h b/Engine/Source/Developer/WorldPartitionHLODUtilities/Public/IWorldPartitionHLODUtilities.h
--- a/Engine/Source/Developer/WorldPartitionHLODUtilities/Public/IWorldPartitionHLODUtilities.h
+++ b/Engine/Source/Developer/WorldPartitionHLODUtilities/Public/IWorldPartitionHLODUtilities.h
@@ -11,6 +11,7 @@
 class AActor;
 class UWorldPartition;
 class UHLODLayer;
+class UHLODBuilder;
 class UHLODBuilderSettings;
 class AWorldPartitionHLOD;
 
@@ -62,6 +63,14 @@ public:
 	 */
 	virtual uint32 BuildHLOD(AWorldPartitionHLOD* InHLODActor) = 0;
 
+	/**
+	 * Retrieve the HLOD Builder class to use for the given HLODLayer.
+	 * 
+	 * @param	InHLODLayer		HLODLayer
+	 * @return The HLOD builder subclass to use for building HLODs for the provided HLOD layer.
+	 */
+	virtual TSubclassOf<UHLODBuilder> GetHLODBuilderClass(const UHLODLayer* InHLODLayer) = 0;
+
 	/**
 	 * Create the HLOD builder settings for the provided HLOD layer object. The type of settings created will depend on the HLOD layer type.
 	 *
diff --git a/Engine/Source/Developer/WorldPartitionHLODUtilities/Public/WorldPartitionHLODUtilities.h b/Engine/Source/Developer/WorldPartitionHLODUtilities/Public/WorldPartitionHLODUtilities.h
--- a/Engine/Source/Developer/WorldPartitionHLODUtilities/Public/WorldPartitionHLODUtilities.h
+++ b/Engine/Source/Developer/WorldPartitionHLODUtilities/Public/WorldPartitionHLODUtilities.h
@@ -12,5 +12,6 @@ class FWorldPartitionHLODUtilities : public IWorldPartitionHLODUtilities
 public:
 	virtual TArray<AWorldPartitionHLOD*> CreateHLODActors(FHLODCreationContext& InCreationContext, const FHLODCreationParams& InCreationParams, const TSet<FActorInstance>& InActors, const TArray<const UDataLayer*>& InDataLayers) override;
 	virtual uint32 BuildHLOD(AWorldPartitionHLOD* InHLODActor) override;
+	virtual TSubclassOf<UHLODBuilder> GetHLODBuilderClass(const UHLODLayer* InHLODLayer) override;
 	virtual UHLODBuilderSettings* CreateHLODBuilderSettings(UHLODLayer* InHLODLayer) override;
 };
diff --git a/Engine/Source/Runtime/Engine/Private/WorldPartition/HLOD/HLODLayer.cpp b/Engine/Source/Runtime/Engine/Private/WorldPartition/HLOD/HLODLayer.cpp
--- a/Engine/Source/Runtime/Engine/Private/WorldPartition/HLOD/HLODLayer.cpp
+++ b/Engine/Source/Runtime/Engine/Private/WorldPartition/HLOD/HLODLayer.cpp
@@ -80,6 +80,17 @@ UHLODLayer* UHLODLayer::GetHLODLayer(const FWorldPartitionActorDesc& InActorDesc
 	return GetHLODLayer(FWorldPartitionActorDescView(&InActorDesc), InWorldPartition);
 }
 
+bool UHLODLayer::DoesRequireWarmup() const
+{
+	IWorldPartitionHLODUtilities* WPHLODUtilities = FModuleManager::Get().LoadModuleChecked<IWorldPartitionHLODUtilitiesModule>("WorldPartitionHLODUtilities").GetUtilities();
+	if (WPHLODUtilities)
+	{
+		return WPHLODUtilities->GetHLODBuilderClass(this)->GetDefaultObject<UHLODBuilder>()->RequiresWarmup();
+	}
+
+	return false;
+}
+
 void UHLODLayer::PostLoad()
 {
 	Super::PostLoad();
diff --git a/Engine/Source/Runtime/Engine/Public/WorldPartition/HLOD/HLODLayer.h b/Engine/Source/Runtime/Engine/Public/WorldPartition/HLOD/HLODLayer.h
--- a/Engine/Source/Runtime/Engine/Public/WorldPartition/HLOD/HLODLayer.h
+++ b/Engine/Source/Runtime/Engine/Public/WorldPartition/HLOD/HLODLayer.h
@@ -53,7 +53,7 @@ public:
 	const TSoftObjectPtr<UHLODLayer>& GetParentLayer() const;
 	bool IsAlwaysLoaded() const { return bAlwaysLoaded; }
 
-	bool DoesRequireWarmup() const { return HLODBuilderClass ? HLODBuilderClass->GetDefaultObject<UHLODBuilder>()->RequiresWarmup() : true; }
+	bool DoesRequireWarmup() const;
 
 	static FName GetRuntimeGridName(uint32 InLODLevel, int32 InCellSize, double InLoadingRange);
 #endif
