Fixed always loaded HLOD actors getting dirtied/saved/submitted each time the HLOD build process is ran.
These HLOD are not rebuilt (hash computation is still the same), so the HLOD generation is fast, but still, this means everyone has to sync those dummy changes everyday.

diff --git a/Engine/Source/Developer/WorldPartitionHLODUtilities/Private/WorldPartitionHLODUtilities.cpp b/Engine/Source/Developer/WorldPartitionHLODUtilities/Private/WorldPartitionHLODUtilities.cpp
--- a/Engine/Source/Developer/WorldPartitionHLODUtilities/Private/WorldPartitionHLODUtilities.cpp
+++ b/Engine/Source/Developer/WorldPartitionHLODUtilities/Private/WorldPartitionHLODUtilities.cpp
@@ -268,9 +268,11 @@ TArray<AWorldPartitionHLOD*> FWorldPartitionHLODUtilities::CreateHLODActors(FHLO
 		}
 
 		// Runtime grid placement
-		if (HLODActor->GetGridPlacement() != SubActorsInfo.GridPlacement)
+		// HLOD that are always loaded will not take the SubActorsInfo.GridPlacement into account
+		EActorGridPlacement ExpectedGridPlacement = HLODLayer->IsAlwaysLoaded() ? EActorGridPlacement::AlwaysLoaded : SubActorsInfo.GridPlacement;
+		if (HLODActor->GetGridPlacement() != ExpectedGridPlacement)
 		{
-			HLODActor->SetGridPlacement(SubActorsInfo.GridPlacement);
+			HLODActor->SetGridPlacement(ExpectedGridPlacement);
 			bIsDirty = true;
 		}
 
